Sequencer: Decorations in property panels.

Add the ability for decorations (for now called 'modifiers' in the UI) to be declared compatible with certain decorations containers (e.g. section, track) and therefore to have their properties appear on the property panels for those containers. Users can choose to add or remove these compatible decorations in the properties panel. We may later add better, more obvious UX for adding these.

#Synced-CL 49622961

diff --git a/Engine/Source/Editor/MovieSceneTools/Private/Decorations/MovieSceneDecorationContainerCustomization.cpp b/Engine/Source/Editor/MovieSceneTools/Private/Decorations/MovieSceneDecorationContainerCustomization.cpp
new file mode 100644
--- /dev/null
+++ b/Engine/Source/Editor/MovieSceneTools/Private/Decorations/MovieSceneDecorationContainerCustomization.cpp
@@ -0,0 +1,199 @@
+// Copyright Epic Games, Inc. All Rights Reserved.
+
+#include "Decorations/MovieSceneDecorationContainerCustomization.h"
+#include "Decorations/MovieSceneDecorationContainer.h"
+#include "PropertyHandle.h"
+#include "DetailCategoryBuilder.h"
+#include "DetailLayoutBuilder.h"
+#include "DetailWidgetRow.h"
+#include "IDetailChildrenBuilder.h"
+#include "IMovieSceneModule.h"
+#include "IPropertyUtilities.h"
+#include "Widgets/Input/SButton.h"
+#include "Widgets/Input/SComboButton.h"
+#include "Widgets/Text/STextBlock.h"
+#include "Framework/MultiBox/MultiBoxBuilder.h"
+#include "ScopedTransaction.h"
+#include "Templates/SubclassOf.h"
+
+#define LOCTEXT_NAMESPACE "MovieSceneDecorationContainer"
+
+TSharedRef<IPropertyTypeCustomization> FMovieSceneDecorationContainerCustomization::MakeInstance()
+{
+	return MakeShared<FMovieSceneDecorationContainerCustomization>();
+}
+
+void FMovieSceneDecorationContainerCustomization::CustomizeHeader(TSharedRef<IPropertyHandle> PropertyHandle, FDetailWidgetRow& HeaderRow, IPropertyTypeCustomizationUtils& CustomizationUtils)
+{
+	// We don't show this property, just the children, and don't want to show the header
+	PropertyHandle->MarkHiddenByCustomization();
+}
+
+void FMovieSceneDecorationContainerCustomization::CustomizeChildren(TSharedRef<IPropertyHandle> PropertyHandle, IDetailChildrenBuilder& ChildBuilder, IPropertyTypeCustomizationUtils& CustomizationUtils)
+{
+	PropertyUtils = CustomizationUtils.GetPropertyUtilities();
+
+	// Add categories to parent layout to prevent an extra layer of headers
+	IDetailLayoutBuilder& LayoutBuilder = ChildBuilder.GetParentCategory().GetParentLayout();
+
+	TArray<UObject*> OuterObjects;
+	PropertyHandle->GetOuterObjects(OuterObjects);
+
+	if (OuterObjects.Num() > 0)
+	{
+		if (UMovieSceneDecorationContainerObject* Container = Cast<UMovieSceneDecorationContainerObject>(OuterObjects[0]))
+		{
+			// Get all compatible (optional/removable) decorations for this container
+			TSet<UClass*> CompatibleClasses;
+			IMovieSceneModule::Get().GetCompatibleDecorationsForContainer(Container, CompatibleClasses);
+
+			for (UObject* DecorationObject : Container->GetDecorations())
+			{
+				if (DecorationObject)
+				{
+					UClass* DecorationClass = DecorationObject->GetClass();
+
+					// Add category for the decoration object based on its class display name. Sort it to the bottom using uncommon category priority.
+					IDetailCategoryBuilder& DecorationCategory = LayoutBuilder.EditCategory(*DecorationClass->GetName(),
+						DecorationClass->GetDisplayNameText(), ECategoryPriority::Uncommon);
+
+					// If this decoration is from the compatible list, add a remove button (b/c that means it's optional and the user has added it)
+					if (CompatibleClasses.Contains(DecorationClass))
+					{
+						DecorationCategory.HeaderContent(
+							SNew(SButton)
+							.OnClicked_Lambda([this, Container, DecorationClass]()
+							{
+								OnRemoveDecoration(Container, DecorationClass);
+								return FReply::Handled();
+							})
+							.ToolTipText(FText::Format(LOCTEXT("RemoveModifiersTooltip", "Remove {0} modifier from this container"), DecorationClass->GetDisplayNameText()))
+							[
+								SNew(STextBlock)
+								.Text(LOCTEXT("RemoveModifier", "Remove Modifier"))
+							]
+						);
+					}
+
+					// Add the object as external
+					DecorationCategory.AddExternalObjects(TArray<UObject*>({DecorationObject}), EPropertyLocation::Default,
+						FAddPropertyParams().HideRootObjectNode(true).AllowChildren(true));
+				}
+			}
+
+			// Add UI for adding compatible decorations
+			AddCompatibleDecorationsUI(Container, LayoutBuilder);
+		}
+	}
+}
+
+void FMovieSceneDecorationContainerCustomization::AddCompatibleDecorationsUI(UMovieSceneDecorationContainerObject* Container, IDetailLayoutBuilder& LayoutBuilder)
+{
+	// Get all compatible decorations from the MovieSceneModule registry and the container's own implementation
+	TSet<UClass*> CompatibleClasses;
+	IMovieSceneModule::Get().GetCompatibleDecorationsForContainer(Container, CompatibleClasses);
+
+	// Filter out already added decorations (only one of each type allowed)
+	TArrayView<const TObjectPtr<UObject>> ExistingDecorations = Container->GetDecorations();
+	for (const TObjectPtr<UObject>& Decoration : ExistingDecorations)
+	{
+		if (Decoration)
+		{
+			CompatibleClasses.Remove(Decoration->GetClass());
+		}
+	}
+
+	// Only show the UI if there are compatible decorations to add
+	if (CompatibleClasses.Num() == 0)
+	{
+		return;
+	}
+
+	// Create a category for adding decorations
+	IDetailCategoryBuilder& AddDecorationsCategory = LayoutBuilder.EditCategory("Modifiers",
+		LOCTEXT("AddModifiersCategory", "Modifiers"), ECategoryPriority::Uncommon);
+
+	AddDecorationsCategory.AddCustomRow(LOCTEXT("AddModifierRow", "Add Modifier"))
+		.WholeRowContent()
+		[
+			SNew(SHorizontalBox)
+			+ SHorizontalBox::Slot()
+			.AutoWidth()
+			.Padding(2.0f)
+			[
+				SNew(SComboButton)
+				.OnGetMenuContent(this, &FMovieSceneDecorationContainerCustomization::OnGetAddDecorationMenuContent, Container, CompatibleClasses)
+				.ButtonContent()
+				[
+					SNew(STextBlock)
+					.Text(LOCTEXT("AddModifier", "Add Modifier"))
+				]
+			]
+		];
+}
+
+TSharedRef<SWidget> FMovieSceneDecorationContainerCustomization::OnGetAddDecorationMenuContent(UMovieSceneDecorationContainerObject* Container, TSet<UClass*> CompatibleClasses)
+{
+	FMenuBuilder MenuBuilder(true, nullptr);
+
+	for (UClass* DecorationClass : CompatibleClasses)
+	{
+		if (!DecorationClass)
+		{
+			continue;
+		}
+
+		FText DisplayName = DecorationClass->GetDisplayNameText();
+		FText ToolTip = FText::Format(LOCTEXT("AddModifierTooltip", "Add {0} modifier to this container"), DisplayName);
+
+		MenuBuilder.AddMenuEntry(
+			DisplayName,
+			ToolTip,
+			FSlateIcon(),
+			FUIAction(
+				FExecuteAction::CreateSP(this, &FMovieSceneDecorationContainerCustomization::OnAddDecoration, Container, DecorationClass)
+			)
+		);
+	}
+
+	return MenuBuilder.MakeWidget();
+}
+
+void FMovieSceneDecorationContainerCustomization::OnAddDecoration(UMovieSceneDecorationContainerObject* Container, UClass* DecorationClass)
+{
+	if (!Container || !DecorationClass)
+	{
+		return;
+	}
+
+	FScopedTransaction Transaction(LOCTEXT("AddModifier", "Add Modifier"));
+	Container->Modify();
+
+	UObject* NewDecoration = NewObject<UObject>(Container, DecorationClass, NAME_None, RF_Transactional);
+	Container->AddDecoration(NewDecoration);
+
+	if (PropertyUtils.IsValid())
+	{
+		PropertyUtils.Pin()->RequestForceRefresh();
+	}
+}
+
+void FMovieSceneDecorationContainerCustomization::OnRemoveDecoration(UMovieSceneDecorationContainerObject* Container, UClass* DecorationClass)
+{
+	if (!Container || !DecorationClass)
+	{
+		return;
+	}
+
+	FScopedTransaction Transaction(LOCTEXT("RemoveModifierTransaction", "Remove Modifier"));
+	Container->Modify();
+
+	Container->RemoveDecoration(DecorationClass);
+
+	if (PropertyUtils.IsValid())
+	{
+		PropertyUtils.Pin()->RequestForceRefresh();
+	}
+}
+
+#undef LOCTEXT_NAMESPACE
\ No newline at end of file
diff --git a/Engine/Source/Editor/MovieSceneTools/Public/Decorations/MovieSceneDecorationContainerCustomization.h b/Engine/Source/Editor/MovieSceneTools/Public/Decorations/MovieSceneDecorationContainerCustomization.h
new file mode 100644
--- /dev/null
+++ b/Engine/Source/Editor/MovieSceneTools/Public/Decorations/MovieSceneDecorationContainerCustomization.h
@@ -0,0 +1,37 @@
+// Copyright Epic Games, Inc. All Rights Reserved.
+
+#pragma once
+
+#include "CoreMinimal.h"
+#include "IPropertyTypeCustomization.h"
+
+class IPropertyHandle;
+class IPropertyTypeCustomizationUtils;
+class UMovieSceneDecorationContainerObject;
+class IDetailLayoutBuilder;
+class SWidget;
+
+/**
+ * Detail customization for UMovieSceneDecorationContainerObject
+ */
+class FMovieSceneDecorationContainerCustomization : public IPropertyTypeCustomization
+{
+public:
+	static MOVIESCENETOOLS_API TSharedRef<IPropertyTypeCustomization> MakeInstance();
+
+	// IPropertyTypeCustomization interface
+	virtual void CustomizeHeader(TSharedRef<IPropertyHandle> PropertyHandle, FDetailWidgetRow& HeaderRow, IPropertyTypeCustomizationUtils& CustomizationUtils) override;
+	virtual void CustomizeChildren(TSharedRef<IPropertyHandle> PropertyHandle, IDetailChildrenBuilder& ChildBuilder, IPropertyTypeCustomizationUtils& CustomizationUtils) override;
+	
+private:
+
+	void AddCompatibleDecorationsUI(UMovieSceneDecorationContainerObject* Container, IDetailLayoutBuilder& LayoutBuilder);
+
+	TSharedRef<SWidget> OnGetAddDecorationMenuContent(UMovieSceneDecorationContainerObject* Container, TSet<UClass*> CompatibleClasses);
+
+	void OnAddDecoration(UMovieSceneDecorationContainerObject* Container, UClass* DecorationClass);
+
+	void OnRemoveDecoration(UMovieSceneDecorationContainerObject* Container, UClass* DecorationClass);
+
+	TWeakPtr<IPropertyUtilities> PropertyUtils;
+};
diff --git a/Engine/Source/Editor/Sequencer/Private/SequencerCommonHelpers.cpp b/Engine/Source/Editor/Sequencer/Private/SequencerCommonHelpers.cpp
--- a/Engine/Source/Editor/Sequencer/Private/SequencerCommonHelpers.cpp
+++ b/Engine/Source/Editor/Sequencer/Private/SequencerCommonHelpers.cpp
@@ -3,6 +3,7 @@
 #include "SequencerCommonHelpers.h"
 
 #include "Algo/Reverse.h"
+#include "Decorations/MovieSceneDecorationContainerCustomization.h"
 #include "EntitySystem/IMovieSceneBlenderSystemSupport.h"
 #include "EntitySystem/MovieSceneBlenderSystem.h"
 #include "FrameNumberDetailsCustomization.h"
@@ -443,6 +444,9 @@ void SequencerHelpers::BuildEditSectionMenu(const TWeakPtr<FSequencer>& InWeakSe
 
 		DetailsView->RegisterInstancedCustomPropertyTypeLayout("MovieSceneDirectorBlueprintConditionData", FOnGetPropertyTypeCustomizationInstance::CreateLambda([=]() {
 			return FMovieSceneDirectorBlueprintConditionCustomization::MakeInstance(CurrentScene.Get());}));
+			
+		DetailsView->RegisterInstancedCustomPropertyTypeLayout("MovieSceneDecorationContainer", FOnGetPropertyTypeCustomizationInstance::CreateLambda([]() {
+			return FMovieSceneDecorationContainerCustomization::MakeInstance(); }));
 	
 		DetailsView->SetIsPropertyVisibleDelegate(FIsPropertyVisible::CreateLambda([](const FPropertyAndParent& PropertyAndParent)
 			{			
@@ -554,6 +558,10 @@ void SequencerHelpers::BuildEditTrackMenu(const TWeakPtr<FSequencer>& InWeakSequ
 		DetailsView->RegisterInstancedCustomPropertyTypeLayout("MovieSceneDirectorBlueprintConditionData", FOnGetPropertyTypeCustomizationInstance::CreateLambda([=]() {
 			return FMovieSceneDirectorBlueprintConditionCustomization::MakeInstance(CurrentScene); }));
 
+		DetailsView->RegisterInstancedCustomPropertyTypeLayout("MovieSceneDecorationContainer", FOnGetPropertyTypeCustomizationInstance::CreateLambda([]() {
+			return FMovieSceneDecorationContainerCustomization::MakeInstance(); }));
+
+
 		DetailsView->SetIsPropertyVisibleDelegate(FIsPropertyVisible::CreateLambda([](const FPropertyAndParent& PropertyAndParent)
 			{
 				return FPropertyEditorPermissionList::Get().DoesPropertyPassFilter(PropertyAndParent.Property.GetOwnerStruct(), PropertyAndParent.Property.GetFName());
diff --git a/Engine/Source/Runtime/MovieScene/Private/MovieSceneModule.cpp b/Engine/Source/Runtime/MovieScene/Private/MovieSceneModule.cpp
--- a/Engine/Source/Runtime/MovieScene/Private/MovieSceneModule.cpp
+++ b/Engine/Source/Runtime/MovieScene/Private/MovieSceneModule.cpp
@@ -12,6 +12,7 @@
 #include "Modules/VisualizerDebuggingState.h"
 #include "EntitySystem/BuiltInComponentTypes.h"
 #include "EntitySystem/MovieSceneEntityManager.h"
+#include "Decorations/MovieSceneDecorationContainer.h"
 
 DEFINE_LOG_CATEGORY(LogMovieScene);
 DEFINE_LOG_CATEGORY(LogMovieSceneECS);
@@ -140,10 +141,69 @@ public:
 	{
 		return ModuleHandle;
 	}
+	
+	virtual void RegisterCompatibleDecoration(TSubclassOf<UMovieSceneDecorationContainerObject> ContainerClass, TSubclassOf<UObject> DecorationClass) override
+	{
+		if (!ContainerClass || !DecorationClass)
+		{
+			return;
+		}
+
+		TSet<TSubclassOf<UObject>>& Decorations = CompatibleDecorations.FindOrAdd(ContainerClass);
+		Decorations.Add(DecorationClass);
+	}
+
+	virtual void UnregisterCompatibleDecoration(TSubclassOf<UMovieSceneDecorationContainerObject> ContainerClass, TSubclassOf<UObject> DecorationClass) override
+	{
+		if (!ContainerClass || !DecorationClass)
+		{
+			return;
+		}
+
+		if (TSet<TSubclassOf<UObject>>* Decorations = CompatibleDecorations.Find(ContainerClass))
+		{
+			Decorations->Remove(DecorationClass);
+			if (Decorations->IsEmpty())
+			{
+				CompatibleDecorations.Remove(ContainerClass);
+			}
+		}
+	}
+
+	virtual void GetCompatibleDecorationsForContainer(const UMovieSceneDecorationContainerObject* Container, TSet<UClass*>& OutClasses) const override
+	{
+		if (!Container)
+		{
+			return;
+		}
+
+		// First get the decorations from the container's own implementation
+		Container->GetCompatibleUserDecorations(OutClasses);
+
+		// Then add any registered decorations for this class and its parent classes
+		UClass* ContainerClass = Container->GetClass();
+		while (ContainerClass && ContainerClass->IsChildOf(UMovieSceneDecorationContainerObject::StaticClass()))
+		{
+			if (const TSet<TSubclassOf<UObject>>* RegisteredDecorations = CompatibleDecorations.Find(ContainerClass))
+			{
+				for (const TSubclassOf<UObject>& DecorationClass : *RegisteredDecorations)
+				{
+					if (DecorationClass)
+					{
+						OutClasses.Add(DecorationClass);
+					}
+				}
+			}
+			ContainerClass = ContainerClass->GetSuperClass();
+		}
+	}
 
 private:
 	TSharedPtr<FMovieSceneModule> ModuleHandle;
 	TMap<FName, FMovieSceneEvaluationGroupParameters> EvaluationGroupParameters;
+
+	// Map from container class to set of compatible decoration classes
+	TMap<TSubclassOf<UMovieSceneDecorationContainerObject>, TSet<TSubclassOf<UObject>>> CompatibleDecorations;
 };
 
 
diff --git a/Engine/Source/Runtime/MovieScene/Public/Decorations/MovieSceneDecorationContainer.h b/Engine/Source/Runtime/MovieScene/Public/Decorations/MovieSceneDecorationContainer.h
--- a/Engine/Source/Runtime/MovieScene/Public/Decorations/MovieSceneDecorationContainer.h
+++ b/Engine/Source/Runtime/MovieScene/Public/Decorations/MovieSceneDecorationContainer.h
@@ -200,6 +200,6 @@ protected:
 protected:
 
 	/** Array of decorations for this movie scene */
-	UPROPERTY()
+	UPROPERTY(EditAnywhere, Category = "Modifiers")
 	FMovieSceneDecorationContainer Decorations;
 };
diff --git a/Engine/Source/Runtime/MovieScene/Public/IMovieSceneModule.h b/Engine/Source/Runtime/MovieScene/Public/IMovieSceneModule.h
--- a/Engine/Source/Runtime/MovieScene/Public/IMovieSceneModule.h
+++ b/Engine/Source/Runtime/MovieScene/Public/IMovieSceneModule.h
@@ -5,6 +5,7 @@
 #include "CoreMinimal.h"
 #include "Modules/ModuleInterface.h"
 #include "Modules/ModuleManager.h"
+#include "Decorations/MovieSceneDecorationContainer.h"
 
 struct FMovieSceneEvaluationGroupParameters
 {
@@ -69,4 +70,21 @@ public:
 	 * @note: resulting weak ptr should not be used to hold persistent strong references
 	 */
 	virtual TWeakPtr<IMovieSceneModule> GetWeakPtr() = 0;
+	
+	/**
+	 * Register a compatible decoration class for a specific container type. 
+	 * This allows modules to externally register compatible decorations for a container type without the container type itself explicitly defining that compatibility.
+	 */
+	virtual void RegisterCompatibleDecoration(TSubclassOf<UMovieSceneDecorationContainerObject> ContainerClass, TSubclassOf<UObject> DecorationClass) = 0;
+
+	/**
+	 * Unregister a compatible decoration class for a specific container type.
+	 */
+	virtual void UnregisterCompatibleDecoration(TSubclassOf<UMovieSceneDecorationContainerObject> ContainerClass, TSubclassOf<UObject> DecorationClass) = 0;
+
+	/**
+	 * Get all compatible decoration classes for a specific container instance.
+	 * This includes both decorations internally marked compatible by a container type and externally registered decorations.
+	 */
+	virtual void GetCompatibleDecorationsForContainer(const UMovieSceneDecorationContainerObject* Container, TSet<UClass*>& OutClasses) const = 0;
 };
