Space Switching and baking: fix out of bound crash

#rnx
#Synced-CL 49723777

diff --git a/Engine/Source/Editor/MovieSceneTools/Public/Constraints/MovieSceneConstraintChannelHelper.inl b/Engine/Source/Editor/MovieSceneTools/Public/Constraints/MovieSceneConstraintChannelHelper.inl
--- a/Engine/Source/Editor/MovieSceneTools/Public/Constraints/MovieSceneConstraintChannelHelper.inl
+++ b/Engine/Source/Editor/MovieSceneTools/Public/Constraints/MovieSceneConstraintChannelHelper.inl
@@ -437,44 +437,94 @@ TArray<FFrameNumber> FMovieSceneConstraintChannelHelper::GetTransformTimes(
 }
 
 template< typename ChannelType >
- void FMovieSceneConstraintChannelHelper::DeleteTransformTimes(
-	 const TArrayView<ChannelType*>& InChannels,
-	 const FFrameNumber& StartTime,
-	 const FFrameNumber& EndTime,
-	 EMovieSceneTransformChannel InChannelsToKey)
+void FMovieSceneConstraintChannelHelper::DeleteTransformTimes(
+	const TArrayView<ChannelType*>& InChannels,
+	const FFrameNumber& StartTime,
+	const FFrameNumber& EndTime,
+	EMovieSceneTransformChannel InChannelsToKey)
 {
-	 TRange<FFrameNumber> WithinRange(0, 0);
-	 WithinRange.SetLowerBoundValue(StartTime);
-	 WithinRange.SetUpperBoundValue(EndTime);
-	 TArray<FFrameNumber> KeyTimes;
-	 TArray<FKeyHandle> KeyHandles;
-
-	 const bool bKeyTranslation = EnumHasAllFlags(InChannelsToKey, EMovieSceneTransformChannel::Translation);
-	 const bool bKeyRotation = EnumHasAllFlags(InChannelsToKey, EMovieSceneTransformChannel::Rotation);
-	 const bool bKeyScale = EnumHasAllFlags(InChannelsToKey, EMovieSceneTransformChannel::Scale);
-
-	 TArray<uint32> ChannelsIndexToKey;
-	 if (bKeyTranslation)
-	 {
-		 ChannelsIndexToKey.Append({ 0,1,2 });
-	 }
-	 if (bKeyRotation)
-	 {
-		 ChannelsIndexToKey.Append({ 3,4,5 });
-	 }
-	 if (bKeyScale)
-	 {
-		 ChannelsIndexToKey.Append({ 6,7,8 });
-	 }
-	 for (const int32 ChannelIndex : ChannelsIndexToKey)
-	 {
-		 ChannelType* Channel = InChannels[ChannelIndex];
-		 TMovieSceneChannelData<typename ChannelType::ChannelValueType> Data = Channel->GetData();
-		 KeyTimes.SetNum(0);
-		 KeyHandles.SetNum(0);
-		 Data.GetKeys(WithinRange, &KeyTimes, &KeyHandles);
-		 Data.DeleteKeys(KeyHandles);
-	 }
+	// nothing to process or invalid channels array 
+	const int32 NumChannels = InChannels.Num();
+	if (NumChannels == 0 || NumChannels > 9)
+	{
+		return;
+	}
+	
+	TRange<FFrameNumber> WithinRange(0, 0);
+	WithinRange.SetLowerBoundValue(StartTime);
+	WithinRange.SetUpperBoundValue(EndTime);
+	
+	auto ProcessAllChannels = [&WithinRange, InChannels]()
+	{
+		TArray<FFrameNumber> KeyTimes;
+		TArray<FKeyHandle> KeyHandles;
+			
+		for (ChannelType* Channel : InChannels)
+		{
+			TMovieSceneChannelData<typename ChannelType::ChannelValueType> Data = Channel->GetData();
+			KeyTimes.SetNum(0);
+			KeyHandles.SetNum(0);
+			Data.GetKeys(WithinRange, &KeyTimes, &KeyHandles);
+			Data.DeleteKeys(KeyHandles);
+		}
+	};
+	
+	// process all channels
+	if (InChannelsToKey == EMovieSceneTransformChannel::AllTransform)
+	{
+		return ProcessAllChannels();
+	}
+
+	const bool bProcessTranslation = EnumHasAllFlags(InChannelsToKey, EMovieSceneTransformChannel::Translation);
+	const bool bProcessRotation = EnumHasAllFlags(InChannelsToKey, EMovieSceneTransformChannel::Rotation);
+	const bool bProcessScale = EnumHasAllFlags(InChannelsToKey, EMovieSceneTransformChannel::Scale);
+	if (!(bProcessTranslation || bProcessRotation || bProcessScale))
+	{
+		// nothing to process
+		return;
+	}
+
+	TArray<uint32> ChannelsIndexToProcess;
+	if (bProcessTranslation)
+	{
+		ChannelsIndexToProcess.Append({0, 1, 2});
+	}
+	if (bProcessRotation)
+	{
+		ChannelsIndexToProcess.Append({3, 4, 5});
+	}
+	if (bProcessScale)
+	{
+		ChannelsIndexToProcess.Append({6, 7, 8});
+	}
+	const int32 NumChannelToProcess = ChannelsIndexToProcess.Num();
+	
+	// process all channels
+	if (NumChannelToProcess == NumChannels)
+	{
+		return ProcessAllChannels();
+	}
+	
+	if (!ensureAlways(NumChannels == 9))
+	{
+		// if NumChannelToProcess != NumChannels and NumChannels is lower than 9 then we don't know what to remove...
+		return;
+	}
+	
+	TArray<FFrameNumber> KeyTimes;
+	TArray<FKeyHandle> KeyHandles;
+	for (const int32 ChannelIndex : ChannelsIndexToProcess)
+	{
+		if (InChannels.IsValidIndex(ChannelIndex))
+		{
+			ChannelType* Channel = InChannels[ChannelIndex];
+			TMovieSceneChannelData<typename ChannelType::ChannelValueType> Data = Channel->GetData();
+			KeyTimes.SetNum(0);
+			KeyHandles.SetNum(0);
+			Data.GetKeys(WithinRange, &KeyTimes, &KeyHandles);
+			Data.DeleteKeys(KeyHandles);
+		}
+	}
 }
 
 template<typename ChannelType>
