Anim In Engine: Global timewarps get baked correctly, revert change to TransformTime to Global since that will invert the time warp.

#Synced-CL 49854396

diff --git a/Engine/Plugins/Animation/ControlRig/Source/ControlRigEditor/Private/ControlRigSequencerEditorLibrary.cpp b/Engine/Plugins/Animation/ControlRig/Source/ControlRigEditor/Private/ControlRigSequencerEditorLibrary.cpp
--- a/Engine/Plugins/Animation/ControlRig/Source/ControlRigEditor/Private/ControlRigSequencerEditorLibrary.cpp
+++ b/Engine/Plugins/Animation/ControlRig/Source/ControlRigEditor/Private/ControlRigSequencerEditorLibrary.cpp
@@ -1206,7 +1206,7 @@ bool UControlRigSequencerEditorLibrary::ExportAnimSequenceFromSequencer(UAnimSeq
 			AnimSequence->SetSkeleton(SkeletalMeshComp->GetSkeletalMeshAsset()->GetSkeleton());
 
 			UMovieSceneSequence* RootMovieSceneSequence = ParentSequencer->GetRootMovieSceneSequence();
-			FMovieSceneSequenceTransform RootToLocalTransform = ParentSequencer->GetFocusedMovieSceneSequenceTransform();
+			FMovieSceneSequenceTransform RootToLocalTransform = ParentSequencer->GetFocusedMovieSceneSequenceTransformUnwarped();
 			FAnimExportSequenceParameters AESQ;
 			AESQ.Player = Player;
 			AESQ.RootToLocalTransform = RootToLocalTransform;
@@ -1631,7 +1631,7 @@ TArray<FArrayOfRigControlValues> UControlRigSequencerEditorLibrary::GetControlVa
 	{
 		Player = Sequencer.Get();
 		Template = Sequencer->GetFocusedTemplateID();
-		RootToLocalTransform = Sequencer->GetFocusedMovieSceneSequenceTransform();
+		RootToLocalTransform = Sequencer->GetFocusedMovieSceneSequenceTransformUnwarped();
 	}
 	else
 	{
@@ -1731,7 +1731,7 @@ static bool GetControlRigValues(ISequencer* Sequencer, UControlRig* ControlRig,
 	if (Sequencer->GetFocusedMovieSceneSequence())
 	{
 		FMovieSceneSequenceIDRef Template = Sequencer->GetFocusedTemplateID();
-		FMovieSceneSequenceTransform RootToLocalTransform = Sequencer->GetFocusedMovieSceneSequenceTransform();
+		FMovieSceneSequenceTransform RootToLocalTransform = Sequencer->GetFocusedMovieSceneSequenceTransformUnwarped();
 		bool bGood = LocalGetControlRigControlValues(Sequencer, Sequencer->GetFocusedMovieSceneSequence(), Template, RootToLocalTransform,
 			ControlRig, ControlName,TimeUnit, Frames, OutValues);
 		Sequencer->RequestEvaluate();
@@ -1749,7 +1749,7 @@ static bool GetControlRigValue(ISequencer* Sequencer, UControlRig* ControlRig, c
 		Frames.Add(Frame);
 		TArray<FRigControlValue> OutValues;
 		FMovieSceneSequenceIDRef Template = Sequencer->GetFocusedTemplateID();
-		FMovieSceneSequenceTransform RootToLocalTransform = Sequencer->GetFocusedMovieSceneSequenceTransform();
+		FMovieSceneSequenceTransform RootToLocalTransform = Sequencer->GetFocusedMovieSceneSequenceTransformUnwarped();
 		bool bVal = LocalGetControlRigControlValues(Sequencer, Sequencer->GetFocusedMovieSceneSequence(), Template, RootToLocalTransform,
 			ControlRig, ControlName,TimeUnit, Frames, OutValues);
 		if (bVal)
@@ -3140,7 +3140,7 @@ bool UControlRigSequencerEditorLibrary::ExportFBXFromControlRigSection(ULevelSeq
 		{
 			if (const TSharedPtr<ISequencer> Sequencer = LevelSequenceEditor->GetSequencer())
 			{
-				RootToLocalTransform = Sequencer->GetFocusedMovieSceneSequenceTransform();
+				RootToLocalTransform = Sequencer->GetFocusedMovieSceneSequenceTransformUnwarped();
 			}
 		}
 	}
diff --git a/Engine/Plugins/Animation/ControlRig/Source/ControlRigEditor/Private/Sequencer/ControlRigParameterTrackEditor.cpp b/Engine/Plugins/Animation/ControlRig/Source/ControlRigEditor/Private/Sequencer/ControlRigParameterTrackEditor.cpp
--- a/Engine/Plugins/Animation/ControlRig/Source/ControlRigEditor/Private/Sequencer/ControlRigParameterTrackEditor.cpp
+++ b/Engine/Plugins/Animation/ControlRig/Source/ControlRigEditor/Private/Sequencer/ControlRigParameterTrackEditor.cpp
@@ -1186,7 +1186,7 @@ void FControlRigParameterTrackEditor::BakeToControlRig(UClass* InClass, FGuid Ob
 			UAnimSequence* TempAnimSequence = NewObject<UAnimSequence>(GetTransientPackage(), NAME_None);
 			TempAnimSequence->SetSkeleton(Skeleton);
 			
-			FMovieSceneSequenceTransform RootToLocalTransform = Sequencer->GetFocusedMovieSceneSequenceTransform();
+			FMovieSceneSequenceTransform RootToLocalTransform = Sequencer->GetFocusedMovieSceneSequenceTransformUnwarped();
 			UAnimSeqExportOption* AnimSeqExportOption = NewObject<UAnimSeqExportOption>(GetTransientPackage(), NAME_None);
 			UBakeToControlRigSettings* BakeSettings = GetMutableDefault<UBakeToControlRigSettings>();
 			AnimSeqExportOption->bTransactRecording = false;
@@ -1458,7 +1458,7 @@ void FControlRigParameterTrackEditor::BakeInvertedPose(UControlRig* InControlRig
 
 	//@sara to do, not sure if you want to key reduce after, but BakeSettings isn't used
 	//UBakeToControlRigSettings* BakeSettings = GetMutableDefault<UBakeToControlRigSettings>();
-	FMovieSceneSequenceTransform RootToLocalTransform = Sequencer->GetFocusedMovieSceneSequenceTransform();
+	FMovieSceneSequenceTransform RootToLocalTransform = Sequencer->GetFocusedMovieSceneSequenceTransformUnwarped();
 
 	const FScopedTransaction Transaction(LOCTEXT("BakeInvertedPose_Transaction", "Bake Inverted Pose"));
 
@@ -6212,7 +6212,7 @@ bool FControlRigParameterTrackEditor::CollapseAllLayers(TSharedPtr<ISequencer>&S
 					FScopedTransaction Transaction(LOCTEXT("CollapseAllSections", "Collapse All Sections"));
 					ParameterSection->Modify();
 					UControlRig* ControlRig = ParameterSection->GetControlRig();
-					FMovieSceneSequenceTransform RootToLocalTransform = SequencerPtr->GetFocusedMovieSceneSequenceTransform();
+					FMovieSceneSequenceTransform RootToLocalTransform = SequencerPtr->GetFocusedMovieSceneSequenceTransformUnwarped();
 
 					FFrameNumber StartFrame = InSettings.StartFrame;
 					FFrameNumber EndFrame = InSettings.EndFrame;
diff --git a/Engine/Plugins/Animation/ControlRig/Source/ControlRigEditor/Private/Tools/ControlRigSnapper.cpp b/Engine/Plugins/Animation/ControlRig/Source/ControlRigEditor/Private/Tools/ControlRigSnapper.cpp
--- a/Engine/Plugins/Animation/ControlRig/Source/ControlRigEditor/Private/Tools/ControlRigSnapper.cpp
+++ b/Engine/Plugins/Animation/ControlRig/Source/ControlRigEditor/Private/Tools/ControlRigSnapper.cpp
@@ -185,7 +185,7 @@ bool FControlRigSnapper::GetControlRigControlTransforms(ISequencer* Sequencer,
 	if (Sequencer->GetFocusedMovieSceneSequence())
 	{
 		FMovieSceneSequenceIDRef Template = Sequencer->GetFocusedTemplateID();
-		FMovieSceneSequenceTransform RootToLocalTransform = Sequencer->GetFocusedMovieSceneSequenceTransform();
+		FMovieSceneSequenceTransform RootToLocalTransform = Sequencer->GetFocusedMovieSceneSequenceTransformUnwarped();
 		TOptional<FFrameNumber> FrameNumber;
 		return LocalGetControlRigControlTransforms(Sequencer, FrameNumber, Sequencer->GetFocusedMovieSceneSequence(), Template, RootToLocalTransform,
 			ControlRig, ControlName, Frames, ParentTransforms, OutTransforms);
@@ -286,7 +286,7 @@ struct FGuidAndActor
 		{
 			UMovieScene* MovieScene = Sequencer->GetFocusedMovieSceneSequence()->GetMovieScene();
 			FFrameRate TickResolution = MovieScene->GetTickResolution();
-			FMovieSceneInverseSequenceTransform LocalToRootTransform = Sequencer->GetFocusedMovieSceneSequenceTransform().Inverse();
+			FMovieSceneInverseSequenceTransform LocalToRootTransform = Sequencer->GetFocusedMovieSceneSequenceTransformUnwarped().Inverse();
 
 			//adjust keys for constraint
 			const FConstraintsManagerController& Controller = FConstraintsManagerController::Get(Actor->GetWorld());
diff --git a/Engine/Plugins/Animation/ControlRig/Source/ControlRigEditor/Private/Tools/EvaluateSequencerTools.cpp b/Engine/Plugins/Animation/ControlRig/Source/ControlRigEditor/Private/Tools/EvaluateSequencerTools.cpp
--- a/Engine/Plugins/Animation/ControlRig/Source/ControlRigEditor/Private/Tools/EvaluateSequencerTools.cpp
+++ b/Engine/Plugins/Animation/ControlRig/Source/ControlRigEditor/Private/Tools/EvaluateSequencerTools.cpp
@@ -208,7 +208,7 @@ namespace AIE
 			return false;
 		}
 		FMovieSceneSequenceIDRef Template = Sequencer->GetFocusedTemplateID();
-		FMovieSceneSequenceTransform RootToLocalTransform = Sequencer->GetFocusedMovieSceneSequenceTransform();
+		FMovieSceneSequenceTransform RootToLocalTransform = Sequencer->GetFocusedMovieSceneSequenceTransformUnwarped();
 		FMovieSceneInverseSequenceTransform LocalToRootTransform = RootToLocalTransform.Inverse();
 		UMovieScene* MovieScene = Sequencer->GetFocusedMovieSceneSequence()->GetMovieScene();
 		const FConstraintsManagerController& Controller = FConstraintsManagerController::Get(World);
@@ -425,7 +425,7 @@ namespace AIE
 
 		TransformSection->Modify();
 
-		FMovieSceneInverseSequenceTransform LocalToRootTransform = Sequencer->GetFocusedMovieSceneSequenceTransform().Inverse();
+		FMovieSceneInverseSequenceTransform LocalToRootTransform = Sequencer->GetFocusedMovieSceneSequenceTransformUnwarped().Inverse();
 		FFrameRate TickResolution = Sequencer->GetFocusedTickResolution();
 
 		//adjust keys for constraints
diff --git a/Engine/Plugins/Editor/SequencerAnimTools/Source/SequencerAnimTools/Private/Sequencer/MovieSceneTransformTrail.cpp b/Engine/Plugins/Editor/SequencerAnimTools/Source/SequencerAnimTools/Private/Sequencer/MovieSceneTransformTrail.cpp
--- a/Engine/Plugins/Editor/SequencerAnimTools/Source/SequencerAnimTools/Private/Sequencer/MovieSceneTransformTrail.cpp
+++ b/Engine/Plugins/Editor/SequencerAnimTools/Source/SequencerAnimTools/Private/Sequencer/MovieSceneTransformTrail.cpp
@@ -1504,17 +1504,17 @@ bool FMovieSceneControlRigTransformTrail::ApplyDelta(const FVector& Pos, const F
 		IMovieScenePlayer* Player = GetSequencer().Get();
 
 		FMovieSceneSequenceTransform RootToLocalTransform = GetSequencer()->GetFocusedMovieSceneSequenceTransform();
-		FMovieSceneInverseSequenceTransform LocalToRootTransform = RootToLocalTransform.Inverse();
 
-		auto EvalControlRig = [&Context,&LocalToRootTransform, Pos, TickResolution, ControlRig,Player, ControlElement, this](FTrailKeyInfo* KeyInfo)
+		auto EvalControlRig = [&Context,&RootToLocalTransform, Pos, TickResolution, ControlRig,Player, ControlElement, this](FTrailKeyInfo* KeyInfo)
 		{
 			if (KeyInfo)
 			{
-				Context.LocalTime = TickResolution.AsSeconds(FFrameTime(KeyInfo->FrameNumber));
-				Context.KeyMask = (uint32)EControlRigContextChannelToKey::Translation;
+				//the keys are in warped time which is what we want to set the control rig values with, but we want to evaluate in unwarped time
 				FFrameTime GlobalTime(KeyInfo->FrameNumber);
-				GlobalTime = LocalToRootTransform.TryTransformTime(GlobalTime).Get(GlobalTime); //player evals in root time so need to go back to it.
-
+				GlobalTime = RootToLocalTransform.TransformTime(GlobalTime);
+				Context.LocalTime = TickResolution.AsSeconds(GlobalTime);
+				Context.KeyMask = (uint32)EControlRigContextChannelToKey::Translation;
+				GlobalTime = RootToLocalTransform.Inverse().TryTransformTime(GlobalTime).Get(GlobalTime);//back to unwarped
 				FMovieSceneContext MovieSceneContext = FMovieSceneContext(FMovieSceneEvaluationRange(GlobalTime, TickResolution), Player->GetPlaybackStatus()).SetHasJumped(true);
 
 				FTransform NewTransform(KeyInfo->Transform);
diff --git a/Engine/Source/Editor/MovieSceneTools/Private/Constraints/ComponentConstraintChannelInterface.cpp b/Engine/Source/Editor/MovieSceneTools/Private/Constraints/ComponentConstraintChannelInterface.cpp
--- a/Engine/Source/Editor/MovieSceneTools/Private/Constraints/ComponentConstraintChannelInterface.cpp
+++ b/Engine/Source/Editor/MovieSceneTools/Private/Constraints/ComponentConstraintChannelInterface.cpp
@@ -483,7 +483,7 @@ void FComponentConstraintChannelInterface::RecomposeTransforms(
 		FDecompositionQuery Query;
 		Query.Object   = SceneComponent;
 		Query.bConvertFromSourceEntityIDs = false;  // We already pass the children entity IDs
-		FMovieSceneSequenceTransform RootToLocalTransform = InSequencer->GetFocusedMovieSceneSequenceTransform();
+		FMovieSceneSequenceTransform RootToLocalTransform = InSequencer->GetFocusedMovieSceneSequenceTransformUnwarped();
 		FMovieSceneInverseSequenceTransform LocalToRootTransform = RootToLocalTransform.Inverse();
 
 		// add keys
diff --git a/Engine/Source/Editor/MovieSceneTools/Private/Constraints/MovieSceneConstraintChannelHelper.cpp b/Engine/Source/Editor/MovieSceneTools/Private/Constraints/MovieSceneConstraintChannelHelper.cpp
--- a/Engine/Source/Editor/MovieSceneTools/Private/Constraints/MovieSceneConstraintChannelHelper.cpp
+++ b/Engine/Source/Editor/MovieSceneTools/Private/Constraints/MovieSceneConstraintChannelHelper.cpp
@@ -119,7 +119,7 @@ void FCompensationEvaluator::ComputeLocalTransforms(UWorld* InWorld, const FEval
 	static constexpr bool bSorted = true;
 	const TArray<ConstraintPtr> AllConstraints = Controller.GetAllConstraints(bSorted);
 
-	FMovieSceneInverseSequenceTransform LocalToRootTransform = Sequencer->GetFocusedMovieSceneSequenceTransform().Inverse();
+	FMovieSceneInverseSequenceTransform LocalToRootTransform = Sequencer->GetFocusedMovieSceneSequenceTransformUnwarped().Inverse();
 
 	const bool bCreatedBeforeCurrentFrame = InEvalParams.OptCurrentFrame.IsSet();
 	const bool bPropagateRelative = InEvalParams.bPropagateRelative;
@@ -339,7 +339,7 @@ void FCompensationEvaluator::ComputeLocalTransformsForBaking(UWorld* InWorld, co
 	const FMovieSceneToolsAnimationBakingWrapper BakingWrapper(MovieScene);
 	BakingWrapper.Start();
 	
-	FMovieSceneInverseSequenceTransform LocalToRootTransform = Sequencer->GetFocusedMovieSceneSequenceTransform().Inverse();
+	FMovieSceneInverseSequenceTransform LocalToRootTransform = Sequencer->GetFocusedMovieSceneSequenceTransformUnwarped().Inverse();
 
 	{
 		const float DisplayRate = static_cast<float>(MovieScene->GetDisplayRate().AsDecimal());
@@ -478,7 +478,7 @@ void FCompensationEvaluator::ComputeLocalTransformsBeforeDeletion(UWorld* InWorl
 	const FMovieSceneToolsAnimationBakingWrapper BakingWrapper(MovieScene);
 	BakingWrapper.Start();
 	
-	FMovieSceneInverseSequenceTransform LocalToRootTransform = Sequencer->GetFocusedMovieSceneSequenceTransform().Inverse();
+	FMovieSceneInverseSequenceTransform LocalToRootTransform = Sequencer->GetFocusedMovieSceneSequenceTransformUnwarped().Inverse();
 
 	for (int32 Index = 0; Index < NumFrames; ++Index)
 	{
@@ -559,7 +559,7 @@ void FCompensationEvaluator::ComputeCompensation(UWorld* InWorld, const TSharedP
 	const FMovieSceneToolsAnimationBakingWrapper BakingWrapper(MovieScene);
 	BakingWrapper.Start();
 
-	FMovieSceneInverseSequenceTransform LocalToRootTransform = InSequencer->GetFocusedMovieSceneSequenceTransform().Inverse();
+	FMovieSceneInverseSequenceTransform LocalToRootTransform = InSequencer->GetFocusedMovieSceneSequenceTransformUnwarped().Inverse();
 
 	auto EvaluateAt = [Handle = this->Handle, InSequencer, &AllConstraints, &BakingWrapper, &LocalToRootTransform](FFrameNumber InFrame)
 	{
@@ -659,7 +659,7 @@ void FCompensationEvaluator::CacheTransforms(UWorld* InWorld, const FEvalParamet
 
 	const FMovieSceneToolsAnimationBakingWrapper BakingWrapper(MovieScene);
 
-	FMovieSceneInverseSequenceTransform LocalToRootTransform = Sequencer->GetFocusedMovieSceneSequenceTransform().Inverse();
+	FMovieSceneInverseSequenceTransform LocalToRootTransform = Sequencer->GetFocusedMovieSceneSequenceTransformUnwarped().Inverse();
 
 	auto EvaluateAt = [&](FFrameNumber InFrame)
 	{
diff --git a/Engine/Source/Editor/MovieSceneTools/Private/MovieSceneToolHelpers.cpp b/Engine/Source/Editor/MovieSceneTools/Private/MovieSceneToolHelpers.cpp
--- a/Engine/Source/Editor/MovieSceneTools/Private/MovieSceneToolHelpers.cpp
+++ b/Engine/Source/Editor/MovieSceneTools/Private/MovieSceneToolHelpers.cpp
@@ -2341,7 +2341,7 @@ private:
 		FMovieSceneSequenceTransform RootToLocalTransform;
 		if (Sequencer.IsValid())
 		{
-			RootToLocalTransform = Sequencer.Pin()->GetFocusedMovieSceneSequenceTransform();
+			RootToLocalTransform = Sequencer.Pin()->GetFocusedMovieSceneSequenceTransformUnwarped();
 		}
 
 		const bool bValid = MovieSceneToolHelpers::ExportFBXFromControlRigChannels(
@@ -5220,7 +5220,7 @@ void MovieSceneToolHelpers::RunSequencerTasks(
 	
 	FFrameRate TickResolution = MovieScene->GetTickResolution();
 	FMovieSceneRootEvaluationTemplateInstance& EvaluationTemplate = Sequencer->GetEvaluationTemplate();
-	const FMovieSceneInverseSequenceTransform InverseSequenceTransform = Sequencer->GetFocusedMovieSceneSequenceTransform().Inverse();
+	const FMovieSceneInverseSequenceTransform InverseSequenceTransform = Sequencer->GetFocusedMovieSceneSequenceTransformUnwarped().Inverse();
 	
 	FMovieSceneToolsAnimationBakingWrapper BakeWrapper(MovieScene);
 	BakeWrapper.Start();
@@ -5280,7 +5280,7 @@ void MovieSceneToolHelpers::RunSequencerTask(
 	
 	FFrameRate TickResolution = MovieScene->GetTickResolution();
 	FMovieSceneRootEvaluationTemplateInstance& EvaluationTemplate = Sequencer->GetEvaluationTemplate();
-	const FMovieSceneInverseSequenceTransform InverseSequenceTransform = Sequencer->GetFocusedMovieSceneSequenceTransform().Inverse();
+	const FMovieSceneInverseSequenceTransform InverseSequenceTransform = Sequencer->GetFocusedMovieSceneSequenceTransformUnwarped().Inverse();
 	
 	FMovieSceneToolsAnimationBakingWrapper BakeWrapper(MovieScene);
 	BakeWrapper.Start();
@@ -5375,7 +5375,7 @@ void MovieSceneToolHelpers::GetActorWorldTransforms(ISequencer* Sequencer, const
 	{
 		if (Sequencer->GetFocusedMovieSceneSequence()->GetMovieScene()->FindTrack<UMovieScene3DTransformTrack>(ObjectHandle))
 		{
-			GetSequencerActorWorldTransforms(Sequencer, Sequencer->GetFocusedMovieSceneSequenceTransform(), Sequencer->GetFocusedMovieSceneSequence(), Template, ActorSelection, Frames, OutWorldTransforms);
+			GetSequencerActorWorldTransforms(Sequencer, Sequencer->GetFocusedMovieSceneSequenceTransformUnwarped(), Sequencer->GetFocusedMovieSceneSequence(), Template, ActorSelection, Frames, OutWorldTransforms);
 			return;
 		}
 	}
@@ -5385,18 +5385,18 @@ void MovieSceneToolHelpers::GetActorWorldTransforms(ISequencer* Sequencer, const
 	{
 		if (Sequencer->GetFocusedMovieSceneSequence()->GetMovieScene()->FindTrack<UMovieScene3DTransformTrack>(ObjectHandle))
 		{
-			GetSequencerActorWorldTransforms(Sequencer, Sequencer->GetFocusedMovieSceneSequenceTransform(), Sequencer->GetFocusedMovieSceneSequence(), Template, ActorSelection, Frames, OutWorldTransforms);
+			GetSequencerActorWorldTransforms(Sequencer, Sequencer->GetFocusedMovieSceneSequenceTransformUnwarped(), Sequencer->GetFocusedMovieSceneSequence(), Template, ActorSelection, Frames, OutWorldTransforms);
 			return;
 		}
 	}
 	if (bAvoidEvaluates)
 	{
 		TArray<FFrameNumber> NoFrame; //this will make sure we don't evaluate
-		GetNonSequencerActorWorldTransforms(Sequencer, Sequencer->GetFocusedMovieSceneSequenceTransform(), Sequencer->GetFocusedMovieSceneSequence(), Template, ActorSelection, NoFrame, OutWorldTransforms);
+		GetNonSequencerActorWorldTransforms(Sequencer, Sequencer->GetFocusedMovieSceneSequenceTransformUnwarped(), Sequencer->GetFocusedMovieSceneSequence(), Template, ActorSelection, NoFrame, OutWorldTransforms);
 	}
 	else
 	{
-		GetNonSequencerActorWorldTransforms(Sequencer, Sequencer->GetFocusedMovieSceneSequenceTransform(), Sequencer->GetFocusedMovieSceneSequence(), Template, ActorSelection, Frames, OutWorldTransforms);
+		GetNonSequencerActorWorldTransforms(Sequencer, Sequencer->GetFocusedMovieSceneSequenceTransformUnwarped(), Sequencer->GetFocusedMovieSceneSequence(), Template, ActorSelection, Frames, OutWorldTransforms);
 	}
 }
 
diff --git a/Engine/Source/Editor/MovieSceneTools/Private/TrackEditors/CommonAnimationTrackEditor.cpp b/Engine/Source/Editor/MovieSceneTools/Private/TrackEditors/CommonAnimationTrackEditor.cpp
--- a/Engine/Source/Editor/MovieSceneTools/Private/TrackEditors/CommonAnimationTrackEditor.cpp
+++ b/Engine/Source/Editor/MovieSceneTools/Private/TrackEditors/CommonAnimationTrackEditor.cpp
@@ -1515,7 +1515,7 @@ void FCommonAnimationTrackEditor::OnSequencerSaved(ISequencer& )
 				UMovieSceneSequence* MovieSceneSequence = SequencerPtr->GetFocusedMovieSceneSequence();
 				UMovieSceneSequence* RootMovieSceneSequence = SequencerPtr->GetRootMovieSceneSequence();
 				FMovieSceneSequenceIDRef Template = SequencerPtr->GetFocusedTemplateID();
-				FMovieSceneSequenceTransform RootToLocalTransform = SequencerPtr->GetFocusedMovieSceneSequenceTransform();
+				FMovieSceneSequenceTransform RootToLocalTransform = SequencerPtr->GetFocusedMovieSceneSequenceTransformUnwarped();
 				//if in subsquence and we want to, turn on ShouldEvaluateSubSequencesInIsolation, use the following object for that
 				{
 					UE::Sequencer::FSequencerBakingSetupRestore RestoreBaking(SequencerPtr);
@@ -1724,7 +1724,7 @@ bool FCommonAnimationTrackEditor::CreateAnimationSequence(const TArray<UObject*>
 				TSharedPtr<ISequencer> ParentSequencer = GetSequencer();
 				UMovieSceneSequence* MovieSceneSequence = ParentSequencer->GetFocusedMovieSceneSequence();
 				UMovieSceneSequence* RootMovieSceneSequence = ParentSequencer->GetRootMovieSceneSequence();
-				FMovieSceneSequenceTransform RootToLocalTransform = ParentSequencer->GetFocusedMovieSceneSequenceTransform();
+				FMovieSceneSequenceTransform RootToLocalTransform = ParentSequencer->GetFocusedMovieSceneSequenceTransformUnwarped();
 				//if in subsquence and we want to, turn on ShouldEvaluateSubSequencesInIsolation, use the following object for that
 				{
 					UE::Sequencer::FSequencerBakingSetupRestore RestoreBaking(ParentSequencer);
diff --git a/Engine/Source/Editor/Sequencer/Private/Sequencer.cpp b/Engine/Source/Editor/Sequencer/Private/Sequencer.cpp
--- a/Engine/Source/Editor/Sequencer/Private/Sequencer.cpp
+++ b/Engine/Source/Editor/Sequencer/Private/Sequencer.cpp
@@ -1252,6 +1252,11 @@ FMovieSceneSequenceTransform FSequencer::GetFocusedMovieSceneSequenceTransform()
 	return RootToWarpedLocalTransform;
 }
 
+FMovieSceneSequenceTransform FSequencer::GetFocusedMovieSceneSequenceTransformUnwarped() const
+{
+	return RootToUnwarpedLocalTransform;
+}
+
 UMovieSceneSequence* FSequencer::GetFocusedMovieSceneSequence() const
 {
 	// the last item is the focused movie scene
diff --git a/Engine/Source/Editor/Sequencer/Private/Sequencer.h b/Engine/Source/Editor/Sequencer/Private/Sequencer.h
--- a/Engine/Source/Editor/Sequencer/Private/Sequencer.h
+++ b/Engine/Source/Editor/Sequencer/Private/Sequencer.h
@@ -723,6 +723,7 @@ public:
 	virtual UMovieSceneSequence* GetRootMovieSceneSequence() const override;
 	virtual UMovieSceneSequence* GetFocusedMovieSceneSequence() const override;
 	virtual FMovieSceneSequenceTransform GetFocusedMovieSceneSequenceTransform() const override;
+	virtual FMovieSceneSequenceTransform GetFocusedMovieSceneSequenceTransformUnwarped() const override;
 	virtual FMovieSceneSequenceTransform GetLocalTimeWarpTransform() const override;
 	virtual FMovieSceneSequenceTransform GetGlobalPlaybackWarpTransform() const override;
 	virtual FMovieSceneRootEvaluationTemplateInstance& GetEvaluationTemplate() override { return RootTemplateInstance; }
diff --git a/Engine/Source/Editor/Sequencer/Public/ISequencer.h b/Engine/Source/Editor/Sequencer/Public/ISequencer.h
--- a/Engine/Source/Editor/Sequencer/Public/ISequencer.h
+++ b/Engine/Source/Editor/Sequencer/Public/ISequencer.h
@@ -267,9 +267,12 @@ public:
 	/** @return Returns the MovieScene that is currently focused for editing by the sequencer.  This can change at any time. */
 	virtual UMovieSceneSequence* GetFocusedMovieSceneSequence() const = 0;
 
-	/**@return Returns the time transform from the focused sequence back to the root*/
+	/**@return Returns the time transform from the focused sequence back to the root, this includes global time warp*/
 	virtual FMovieSceneSequenceTransform GetFocusedMovieSceneSequenceTransform() const = 0;
 
+	/**@return Returns the time transform from the focused sequence back to the root, without timewarp, to be used for example when baking and we need to pass time for evaluation*/
+	virtual FMovieSceneSequenceTransform GetFocusedMovieSceneSequenceTransformUnwarped() const = 0;
+
 	/**@return Returns the timewrap transform that local time to warped-local-time */
 	virtual FMovieSceneSequenceTransform GetLocalTimeWarpTransform() const = 0;
 
