Sequencer: Fix delete category from the category name path. It was deleting the first rather than the last name in the path (which is the actual channel name)

#Synced-CL 49850366

diff --git a/Engine/Source/Editor/MovieSceneTools/Private/Sections/ComponentMaterialParameterSection.cpp b/Engine/Source/Editor/MovieSceneTools/Private/Sections/ComponentMaterialParameterSection.cpp
--- a/Engine/Source/Editor/MovieSceneTools/Private/Sections/ComponentMaterialParameterSection.cpp
+++ b/Engine/Source/Editor/MovieSceneTools/Private/Sections/ComponentMaterialParameterSection.cpp
@@ -251,10 +251,12 @@ bool FComponentMaterialParameterSection::RequestDeleteCategory(const TArray<FNam
 {
 	const FScopedTransaction Transaction(LOCTEXT("DeleteParameter", "Delete parameter"));
 	UMovieSceneComponentMaterialParameterSection* ParameterSection = Cast<UMovieSceneComponentMaterialParameterSection>(WeakSection.Get());
-	if (ParameterSection->Modify())
+	if (ParameterSection && ParameterSection->Modify() && !CategoryNamePath.IsEmpty())
 	{
-		const bool bColorParameterDeleted = ParameterSection->RemoveColorParameter(CategoryNamePath[0]);
-		const bool bScalarParameterDeleted = ParameterSection->RemoveScalarParameter(CategoryNamePath[0]);
+		const FName& ChannelName = CategoryNamePath.Last();
+
+		const bool bColorParameterDeleted = ParameterSection->RemoveColorParameter(ChannelName);
+		const bool bScalarParameterDeleted = ParameterSection->RemoveScalarParameter(ChannelName);
 		return bColorParameterDeleted || bScalarParameterDeleted;
 	}
 	return false;
@@ -268,7 +270,7 @@ bool FComponentMaterialParameterSection::RequestDeleteKeyArea(const TArray<FName
 	{
 		const FScopedTransaction Transaction(LOCTEXT("DeleteParameter", "Delete parameter"));
 		UMovieSceneComponentMaterialParameterSection* ParameterSection = Cast<UMovieSceneComponentMaterialParameterSection>(WeakSection.Get());
-		if (ParameterSection->TryModify())
+		if (ParameterSection && ParameterSection->TryModify())
 		{
 			const bool bColorParameterDeleted = ParameterSection->RemoveColorParameter(KeyAreaNamePath[0]);
 			const bool bScalarParameterDeleted = ParameterSection->RemoveScalarParameter(KeyAreaNamePath[0]);
diff --git a/Engine/Source/Editor/MovieSceneTools/Private/Sections/ParameterSection.cpp b/Engine/Source/Editor/MovieSceneTools/Private/Sections/ParameterSection.cpp
--- a/Engine/Source/Editor/MovieSceneTools/Private/Sections/ParameterSection.cpp
+++ b/Engine/Source/Editor/MovieSceneTools/Private/Sections/ParameterSection.cpp
@@ -254,14 +254,16 @@ bool FParameterSection::RequestDeleteCategory( const TArray<FName>& CategoryName
 {
 	const FScopedTransaction Transaction( LOCTEXT( "DeleteParameter", "Delete parameter" ) );
 	UMovieSceneParameterSection* ParameterSection = Cast<UMovieSceneParameterSection>( WeakSection.Get() );
-	if( ParameterSection->Modify() )
+	if (ParameterSection && ParameterSection->Modify() && !CategoryNamePath.IsEmpty())
 	{
-		const bool bVectorParameterDeleted = ParameterSection->RemoveVectorParameter(CategoryNamePath[0]);
-		const bool bColorParameterDeleted = ParameterSection->RemoveColorParameter(CategoryNamePath[0]);
-		const bool bScalarParameterDeleted = ParameterSection->RemoveScalarParameter(CategoryNamePath[0]);
-		const bool bBoolParameterDeleted = ParameterSection->RemoveBoolParameter(CategoryNamePath[0]);
-		const bool bVector2DParameterDeleted = ParameterSection->RemoveVector2DParameter(CategoryNamePath[0]);
-		const bool bTransformParameterDeleted = ParameterSection->RemoveTransformParameter(CategoryNamePath[0]);
+		const FName& ChannelName = CategoryNamePath.Last();
+
+		const bool bVectorParameterDeleted = ParameterSection->RemoveVectorParameter(ChannelName);
+		const bool bColorParameterDeleted = ParameterSection->RemoveColorParameter(ChannelName);
+		const bool bScalarParameterDeleted = ParameterSection->RemoveScalarParameter(ChannelName);
+		const bool bBoolParameterDeleted = ParameterSection->RemoveBoolParameter(ChannelName);
+		const bool bVector2DParameterDeleted = ParameterSection->RemoveVector2DParameter(ChannelName);
+		const bool bTransformParameterDeleted = ParameterSection->RemoveTransformParameter(ChannelName);
 		return bVectorParameterDeleted || bColorParameterDeleted || bScalarParameterDeleted || bBoolParameterDeleted || bVector2DParameterDeleted || bTransformParameterDeleted;
 	}
 	return false;
@@ -275,7 +277,7 @@ bool FParameterSection::RequestDeleteKeyArea( const TArray<FName>& KeyAreaNamePa
 	{
 		const FScopedTransaction Transaction(LOCTEXT("DeleteParameter", "Delete parameter"));
 		UMovieSceneParameterSection* ParameterSection = Cast<UMovieSceneParameterSection>( WeakSection.Get() );
-		if (ParameterSection->TryModify())
+		if (ParameterSection && ParameterSection->TryModify())
 		{
 			const bool bVectorParameterDeleted = ParameterSection->RemoveVectorParameter(KeyAreaNamePath[0]);
 			const bool bColorParameterDeleted = ParameterSection->RemoveColorParameter(KeyAreaNamePath[0]);
