Sequencer: Automatically create an actor reference track when a pointer to a spawnable object is set

This addresses issues where the user can select a spawnable as the object for an object property and then that spawnable reference will be lost whenever the soft object path changes.

#Synced-CL 49850366

diff --git a/Engine/Source/Editor/MovieSceneTools/Private/TrackEditors/PropertyTrackEditors/ActorReferencePropertyTrackEditor.cpp b/Engine/Source/Editor/MovieSceneTools/Private/TrackEditors/PropertyTrackEditors/ActorReferencePropertyTrackEditor.cpp
--- a/Engine/Source/Editor/MovieSceneTools/Private/TrackEditors/PropertyTrackEditors/ActorReferencePropertyTrackEditor.cpp
+++ b/Engine/Source/Editor/MovieSceneTools/Private/TrackEditors/PropertyTrackEditors/ActorReferencePropertyTrackEditor.cpp
@@ -65,4 +65,25 @@ void FActorReferencePropertyTrackEditor::GenerateKeysFromPropertyChanged(const F
 	}
 }
 
+void FActorReferencePropertyTrackEditor::OnAnimatedPropertyChanged(const FPropertyChangedParams& PropertyChangedParams)
+{
+	// When the changed object is a spawnable, force a manual key so that the change will generate an actor reference track. 
+	// Otherwise, the object will get referenced by an unstable soft object ptr.
+	FPropertyChangedParams NewPropertyChangedParams(PropertyChangedParams);
+
+	AActor* NewReferencedActor = Cast<AActor>(PropertyChangedParams.GetPropertyValue<UObject*>());
+	if (NewReferencedActor != nullptr)
+	{
+		TSharedPtr<ISequencer> SequencerPtr = GetSequencer();
+
+		TOptional<FMovieSceneSpawnableAnnotation> Spawnable = FMovieSceneSpawnableAnnotation::Find(NewReferencedActor);
+		if (Spawnable.IsSet())
+		{
+			NewPropertyChangedParams.KeyMode = ESequencerKeyMode::ManualKey;
+		}
+	}
+
+	FPropertyTrackEditor::OnAnimatedPropertyChanged(NewPropertyChangedParams);
+}
+
 #undef LOCTEXT_NAMESPACE
diff --git a/Engine/Source/Editor/MovieSceneTools/Private/TrackEditors/PropertyTrackEditors/ActorReferencePropertyTrackEditor.h b/Engine/Source/Editor/MovieSceneTools/Private/TrackEditors/PropertyTrackEditors/ActorReferencePropertyTrackEditor.h
--- a/Engine/Source/Editor/MovieSceneTools/Private/TrackEditors/PropertyTrackEditors/ActorReferencePropertyTrackEditor.h
+++ b/Engine/Source/Editor/MovieSceneTools/Private/TrackEditors/PropertyTrackEditors/ActorReferencePropertyTrackEditor.h
@@ -66,4 +66,5 @@ protected:
 	virtual FText GetDisplayName() const override;
 	virtual void ProcessKeyOperation(FFrameNumber InKeyTime, const UE::Sequencer::FKeyOperation& Operation, ISequencer& InSequencer, TArray<UE::Sequencer::FAddKeyResult>* OutResults = nullptr) override;
 	virtual void GenerateKeysFromPropertyChanged(const FPropertyChangedParams& PropertyChangedParams, UMovieSceneSection* SectionToKey, FGeneratedTrackKeys& OutGeneratedKeys) override;
+	virtual void OnAnimatedPropertyChanged(const FPropertyChangedParams& PropertyChangedParams) override;
 };
diff --git a/Engine/Source/Editor/Sequencer/Public/KeyPropertyParams.h b/Engine/Source/Editor/Sequencer/Public/KeyPropertyParams.h
--- a/Engine/Source/Editor/Sequencer/Public/KeyPropertyParams.h
+++ b/Engine/Source/Editor/Sequencer/Public/KeyPropertyParams.h
@@ -124,7 +124,7 @@ public:
 	FPropertyPath StructPathToKey;
 
 	/** Keyframing params */
-	const ESequencerKeyMode KeyMode;
+	ESequencerKeyMode KeyMode;
 };
 
 #undef UE_API
