Sequencer: Only update thumbnails if PIE and the PIE world is bound to the current sequencer (otherwise the thumbnails wouldn't represent anything in the sequence)

Tested that thumbnails don't update during MRG/MRQ

#Synced-CL 50405187

diff --git a/Engine/Source/Editor/MovieSceneTools/Private/TrackEditors/CameraCutTrackEditor.cpp b/Engine/Source/Editor/MovieSceneTools/Private/TrackEditors/CameraCutTrackEditor.cpp
--- a/Engine/Source/Editor/MovieSceneTools/Private/TrackEditors/CameraCutTrackEditor.cpp
+++ b/Engine/Source/Editor/MovieSceneTools/Private/TrackEditors/CameraCutTrackEditor.cpp
@@ -289,7 +289,16 @@ void FCameraCutTrackEditor::Tick(float DeltaTime)
 
 	EMovieScenePlayerStatus::Type PlaybackState = SequencerPin->GetPlaybackStatus();
 
-	if (FSlateThrottleManager::Get().IsAllowingExpensiveTasks() && PlaybackState != EMovieScenePlayerStatus::Playing && PlaybackState != EMovieScenePlayerStatus::Scrubbing)
+	UObject* PlaybackContext = SequencerPin->GetSharedPlaybackState()->GetPlaybackContext();
+	UWorld* World = PlaybackContext ? PlaybackContext->GetWorld() : nullptr;
+
+	const bool bIsInPIEOrSimulate = GEditor->PlayWorld != NULL || GEditor->bIsSimulatingInEditor;
+
+	// Render thumbnails if allow expensive tasks and playback state is not playing or scrubbing and if in a PIE world, the sequence must be bound to it
+	if (FSlateThrottleManager::Get().IsAllowingExpensiveTasks() 
+		&& PlaybackState != EMovieScenePlayerStatus::Playing 
+		&& PlaybackState != EMovieScenePlayerStatus::Scrubbing
+		&& (!bIsInPIEOrSimulate || (bIsInPIEOrSimulate && GEditor->PlayWorld == World)) )
 	{
 		SequencerPin->EnterSilentMode();
 
diff --git a/Engine/Source/Editor/MovieSceneTools/Private/TrackEditors/CinematicShotTrackEditor.cpp b/Engine/Source/Editor/MovieSceneTools/Private/TrackEditors/CinematicShotTrackEditor.cpp
--- a/Engine/Source/Editor/MovieSceneTools/Private/TrackEditors/CinematicShotTrackEditor.cpp
+++ b/Engine/Source/Editor/MovieSceneTools/Private/TrackEditors/CinematicShotTrackEditor.cpp
@@ -162,7 +162,16 @@ void FCinematicShotTrackEditor::Tick(float DeltaTime)
 
 	EMovieScenePlayerStatus::Type PlaybackState = SequencerPin->GetPlaybackStatus();
 
-	if (FSlateThrottleManager::Get().IsAllowingExpensiveTasks() && PlaybackState != EMovieScenePlayerStatus::Playing && PlaybackState != EMovieScenePlayerStatus::Scrubbing)
+	UObject* PlaybackContext = SequencerPin->GetSharedPlaybackState()->GetPlaybackContext();
+	UWorld* World = PlaybackContext ? PlaybackContext->GetWorld() : nullptr;
+
+	const bool bIsInPIEOrSimulate = GEditor->PlayWorld != NULL || GEditor->bIsSimulatingInEditor;
+
+	// Render thumbnails if allow expensive tasks and playback state is not playing or scrubbing and if in a PIE world, the sequence must be bound to it
+	if (FSlateThrottleManager::Get().IsAllowingExpensiveTasks()
+		&& PlaybackState != EMovieScenePlayerStatus::Playing
+		&& PlaybackState != EMovieScenePlayerStatus::Scrubbing
+		&& (!bIsInPIEOrSimulate || (bIsInPIEOrSimulate && GEditor->PlayWorld == World)) )
 	{
 		SequencerPin->EnterSilentMode();
 
