Sequencer: Fix not being able to set just the X, Y, or Z channels of a transform for a skeletal mesh.

40800307 introduced an evaluation when a skeletal mesh property changes so that the skeletal mesh doesn't stay in a t pose when adding keyframes. It skips past Location/Rotation/Scale properties by checking the FPropertyChangedEvent Property. Instead, it should have checked MemberProperty since that will be RelativeLocation/RelativeRotation/RelativeScale3D whenever either the property or the X, Y, Z channel is changed.

#Synced-CL 50405187

diff --git a/Engine/Source/Editor/MovieSceneTools/Private/TrackEditors/SkeletalAnimationTrackEditor.cpp b/Engine/Source/Editor/MovieSceneTools/Private/TrackEditors/SkeletalAnimationTrackEditor.cpp
--- a/Engine/Source/Editor/MovieSceneTools/Private/TrackEditors/SkeletalAnimationTrackEditor.cpp
+++ b/Engine/Source/Editor/MovieSceneTools/Private/TrackEditors/SkeletalAnimationTrackEditor.cpp
@@ -132,10 +132,17 @@ void FSkeletalAnimationTrackEditor::OnPostPropertyChanged(UObject* InObject, FPr
 {
 	if (UE::Sequencer::SkeletalAnimationEditorConstants::CVarEvaluateSkeletalMeshOnPropertyChange->GetBool())
 	{
-		const FName PropertyName = (InPropertyChangedEvent.Property != nullptr) ? InPropertyChangedEvent.Property->GetFName() : NAME_None;
-		if (PropertyName == USceneComponent::GetRelativeLocationPropertyName() ||
-			PropertyName == USceneComponent::GetRelativeRotationPropertyName() ||
-			PropertyName == USceneComponent::GetRelativeScale3DPropertyName())
+		FProperty* MemberPropertyThatChanged = InPropertyChangedEvent.MemberProperty;
+		const FName MemberPropertyName = MemberPropertyThatChanged != NULL ? MemberPropertyThatChanged->GetFName() : NAME_None;
+		if (MemberPropertyName == USceneComponent::GetRelativeLocationPropertyName() ||
+			MemberPropertyName == USceneComponent::GetRelativeRotationPropertyName() ||
+			MemberPropertyName == USceneComponent::GetRelativeScale3DPropertyName())
+		{
+			return;
+		}
+
+		const TSharedPtr<ISequencer> SequencerPtr = GetSequencer();
+		if (!SequencerPtr)
 		{
 			return;
 		}
@@ -149,7 +156,7 @@ void FSkeletalAnimationTrackEditor::OnPostPropertyChanged(UObject* InObject, FPr
 			FFindOrCreateTrackResult TrackResult = FindOrCreateTrackForObject(ObjectHandle, UMovieSceneSkeletalAnimationTrack::StaticClass(), NAME_None, bCreateIfMissing);
 			if (TrackResult.Track)
 			{
-				GetSequencer()->NotifyMovieSceneDataChanged(EMovieSceneDataChangeType::TrackValueChanged);
+				SequencerPtr->NotifyMovieSceneDataChanged(EMovieSceneDataChangeType::TrackValueChanged);
 			}
 		}
 	}
