Fix for animation pose resetting when modifying property on skeletal mesh component.

FSkeletalAnimationTrackEditor::OnPostPropertyChanged now accounts for if a subobject on the actor triggers the property change. It will get the outer actor if the InObject is not itself an acotor.

#Synced-CL 50598455

diff --git a/Engine/Source/Editor/MovieSceneTools/Private/TrackEditors/SkeletalAnimationTrackEditor.cpp b/Engine/Source/Editor/MovieSceneTools/Private/TrackEditors/SkeletalAnimationTrackEditor.cpp
--- a/Engine/Source/Editor/MovieSceneTools/Private/TrackEditors/SkeletalAnimationTrackEditor.cpp
+++ b/Engine/Source/Editor/MovieSceneTools/Private/TrackEditors/SkeletalAnimationTrackEditor.cpp
@@ -149,7 +149,16 @@ void FSkeletalAnimationTrackEditor::OnPostPropertyChanged(UObject* InObject, FPr
 
 		// If the object changed has any animation track, notify sequencer to update because animations tick on their own and sequencer needs to evaluate again
 		const bool bCreateIfMissing = false;
-		FFindOrCreateHandleResult HandleResult = FindOrCreateHandleToObject(InObject, bCreateIfMissing);
+		FFindOrCreateHandleResult HandleResult;
+		if (InObject->IsA(AActor::StaticClass()))
+		{
+			HandleResult = FindOrCreateHandleToObject(InObject, bCreateIfMissing);
+		}
+		else
+		{
+			HandleResult = FindOrCreateHandleToObject(InObject->GetTypedOuter<AActor>(), bCreateIfMissing);
+		}
+		
 		FGuid ObjectHandle = HandleResult.Handle;
 		if (ObjectHandle.IsValid())
 		{
