Anim In Engine: When merging layers use channel key extents and not playback range.

#Synced-CL 50557846

diff --git a/Engine/Source/Editor/MovieSceneTools/Public/MovieSceneToolHelpers.h b/Engine/Source/Editor/MovieSceneTools/Public/MovieSceneToolHelpers.h
--- a/Engine/Source/Editor/MovieSceneTools/Public/MovieSceneToolHelpers.h
+++ b/Engine/Source/Editor/MovieSceneTools/Public/MovieSceneToolHelpers.h
@@ -860,7 +860,7 @@ public:
 	* @param SectionChannelIndex  The index of the channel of this type in the sections.
 	* @param Channels  Set of channels of this type that we will merge together onto the first specified one.
 	* @param Sections Set of corresponding sections for each channel specified.
-	* @param Range Range of time over which we will merge
+	* @param Range Minimal range of time over which we will merge, we will always extend to key extents if possible
 	* @param MergeAlgorithm The algorithm to use to blend the channels together
 	* @param TrackSections Sections in this track
 	* @param Increment Optional increment index, if specfied we bake over the range using the increment value, otherwise by default we just merge onto the existing keys
@@ -897,6 +897,13 @@ public:
 		//base channel we set values on
 		ChannelType* BaseChannel = Channels[0];
 		TMovieSceneChannelData<ChannelValueType> BaseChannelData = BaseChannel->GetData();
+
+		//use following to get key extents
+		TArray<FFrameNumber> KeyTimes;
+		TArray<FKeyHandle> Handles;
+		const FFrameNumber Min = TNumericLimits<FFrameNumber>::Lowest();
+		const FFrameNumber Max = TNumericLimits<FFrameNumber>::Max();
+		TRange<FFrameNumber> MaxRange(Min, Max);
 		//similar but slightly different implementations based upon whether or not doing per key or by frame.
 		//for performance we do this outside of the loop
 		if (Increment != nullptr && *Increment > 0)
@@ -904,8 +911,30 @@ public:
 			//iterate over each frame
 			if (UMovieScene* MovieScene = Sections[0]->GetTypedOuter<UMovieScene>())
 			{
+				TRange<FFrameNumber> RangeToUse = Range;
+				for (int32 ChannelIndex = 0; ChannelIndex < Channels.Num(); ++ChannelIndex)
+				{
+					KeyTimes.Reset();
+					Handles.Reset();
+					ChannelType* Channel = Channels[ChannelIndex];
+					Channel->GetKeys(MaxRange, &KeyTimes, &Handles);
+					if (KeyTimes.Num() > 0)
+					{
+						if (KeyTimes[0] < RangeToUse.GetLowerBoundValue())
+						{
+							RangeToUse.SetLowerBoundValue(KeyTimes[0]);
+						}
+						if (KeyTimes[KeyTimes.Num() - 1] > RangeToUse.GetUpperBoundValue())
+						{
+							RangeToUse.SetUpperBoundValue(KeyTimes[KeyTimes.Num() - 1]);
+						}
+					}
+					
+				}
 				TArray<FFrameNumber> Frames;
-				CalculateFramesBetween(MovieScene, Range.GetLowerBoundValue(), Range.GetUpperBoundValue(), *Increment, Frames);
+				CalculateFramesBetween(MovieScene, RangeToUse.GetLowerBoundValue(), RangeToUse.GetUpperBoundValue(), *Increment, Frames);
+				
+			
 				TArray<TPair< FFrameNumber, CurveValueType>> ValuesToSet;
 				for (int32 ChannelIndex = 0; ChannelIndex < Channels.Num(); ++ChannelIndex)
 				{
@@ -930,9 +959,6 @@ public:
 		}
 		else
 		{
-			//iterate over each key
-			TArray<FFrameNumber> KeyTimes;
-			TArray<FKeyHandle> Handles;
 			//cached set that we set at the end
 			TArray<TPair< FFrameNumber, ChannelValueType>> KeysToSet;
 			for (int32 ChannelIndex = 0; ChannelIndex < Channels.Num(); ++ChannelIndex)
@@ -940,7 +966,7 @@ public:
 				KeyTimes.Reset();
 				Handles.Reset();
 				ChannelType* Channel = Channels[ChannelIndex];
-				Channel->GetKeys(Range, &KeyTimes, &Handles);
+				Channel->GetKeys(MaxRange, &KeyTimes, &Handles);
 				for (int32 FrameIndex = 0; FrameIndex < KeyTimes.Num(); ++FrameIndex)
 				{
 					const FFrameNumber& Frame = KeyTimes[FrameIndex];
