Fix regression on resetting animation pose, where a mesh object containing an animation track would not trigger re-evaluation, due to its outer only being considered. Caused by CL 50607091.

Now both the outer actor and object search for Skeletal animation tracks.

#Synced-CL 50607091

diff --git a/Engine/Source/Editor/MovieSceneTools/Private/TrackEditors/SkeletalAnimationTrackEditor.cpp b/Engine/Source/Editor/MovieSceneTools/Private/TrackEditors/SkeletalAnimationTrackEditor.cpp
--- a/Engine/Source/Editor/MovieSceneTools/Private/TrackEditors/SkeletalAnimationTrackEditor.cpp
+++ b/Engine/Source/Editor/MovieSceneTools/Private/TrackEditors/SkeletalAnimationTrackEditor.cpp
@@ -149,24 +149,23 @@ void FSkeletalAnimationTrackEditor::OnPostPropertyChanged(UObject* InObject, FPr
 
 		// If the object changed has any animation track, notify sequencer to update because animations tick on their own and sequencer needs to evaluate again
 		const bool bCreateIfMissing = false;
-		FFindOrCreateHandleResult HandleResult;
-		if (InObject->IsA(AActor::StaticClass()))
-		{
-			HandleResult = FindOrCreateHandleToObject(InObject, bCreateIfMissing);
-		}
-		else
+		FFindOrCreateHandleResult ObjectHandleResult = FindOrCreateHandleToObject(InObject, bCreateIfMissing);
+		FGuid ObjectHandle = ObjectHandleResult.Handle;
+		
+		bool bTrackFound = ObjectHandle.IsValid() ? FindOrCreateTrackForObject(ObjectHandle, UMovieSceneSkeletalAnimationTrack::StaticClass(), NAME_None, bCreateIfMissing).Track != nullptr : false;
+
+		// A mesh component, or its owning actor can have the animation track, so both need to be searched for.
+		if (!bTrackFound && !InObject->IsA<AActor>())
 		{
-			HandleResult = FindOrCreateHandleToObject(InObject->GetTypedOuter<AActor>(), bCreateIfMissing);
+			ObjectHandleResult = FindOrCreateHandleToObject(InObject->GetTypedOuter<AActor>(), bCreateIfMissing);
+			ObjectHandle = ObjectHandleResult.Handle;
+
+			bTrackFound = ObjectHandle.IsValid() ? FindOrCreateTrackForObject(ObjectHandle, UMovieSceneSkeletalAnimationTrack::StaticClass(), NAME_None, bCreateIfMissing).Track != nullptr : false;
 		}
-		
-		FGuid ObjectHandle = HandleResult.Handle;
-		if (ObjectHandle.IsValid())
+
+		if (bTrackFound)
 		{
-			FFindOrCreateTrackResult TrackResult = FindOrCreateTrackForObject(ObjectHandle, UMovieSceneSkeletalAnimationTrack::StaticClass(), NAME_None, bCreateIfMissing);
-			if (TrackResult.Track)
-			{
-				SequencerPtr->NotifyMovieSceneDataChanged(EMovieSceneDataChangeType::TrackValueChanged);
-			}
+			SequencerPtr->NotifyMovieSceneDataChanged(EMovieSceneDataChangeType::TrackValueChanged);
 		}
 	}
 }
