Sequencer: Add support for Material Slots Overlay Material switching and parameter animation.

#Synced-CL 50659531

diff --git a/Engine/Source/Editor/MovieSceneTools/Private/TrackEditors/MaterialTrackEditor.cpp b/Engine/Source/Editor/MovieSceneTools/Private/TrackEditors/MaterialTrackEditor.cpp
--- a/Engine/Source/Editor/MovieSceneTools/Private/TrackEditors/MaterialTrackEditor.cpp
+++ b/Engine/Source/Editor/MovieSceneTools/Private/TrackEditors/MaterialTrackEditor.cpp
@@ -491,6 +491,12 @@ UMaterialInterface* FComponentMaterialTrackEditor::GetMaterialInterfaceForTrack(
 			return CloudComponent->GetMaterial();
 		}
 		break;
+	case EComponentMaterialType::SlotOverlayMaterial:
+		if (UMeshComponent* Component = Cast<UMeshComponent>(Object))
+		{
+			return Component->GetOverlayMaterial(true, MaterialInfo.MaterialSlotIndex);
+		}
+		break;
 	default:
 		break;
 	}
@@ -568,6 +574,31 @@ void FComponentMaterialTrackEditor::ConstructObjectBindingTrackMenu(FMenuBuilder
 							MenuBuilder.AddMenuEntry(AddOverlayMaterialLabel, AddOverlayMaterialToolTip, FSlateIcon(), AddComponentMaterialAction);
 						}
 					}
+
+					const TArray<TObjectPtr<UMaterialInterface>>& SlotOverlayMaterials = MeshComponent->GetComponentMaterialSlotsOverlayMaterial();
+					for (int32 SlotIndex = 0; SlotIndex < SlotOverlayMaterials.Num(); ++SlotIndex)
+					{
+						UMaterialInterface* SlotOverlayMaterial = SlotOverlayMaterials[SlotIndex];
+						if (!SlotOverlayMaterial)
+						{
+							continue;
+						}
+						FName SlotName = MaterialSlotNames.IsValidIndex(SlotIndex) ? MaterialSlotNames[SlotIndex] : FName();
+						FComponentMaterialInfo MaterialInfo{ SlotName, SlotIndex, EComponentMaterialType::SlotOverlayMaterial };
+						const bool bAlreadyExists = Algo::FindBy(Binding->GetTracks(), MaterialInfo, GetMaterialInfoForTrack) != nullptr;
+						if (bAlreadyExists)
+						{
+							continue;
+						}
+						FUIAction AddComponentMaterialAction(FExecuteAction::CreateRaw(this, &FComponentMaterialTrackEditor::HandleAddComponentMaterialActionExecute, SceneComponent, MaterialInfo));
+						FText AddSlotOverlayLabel = !SlotName.IsNone() ?
+							FText::Format(LOCTEXT("AddSlotOverlayMaterialLabelFormat", "Overlay Slot: {0}"), FText::FromName(SlotName)) :
+							FText::Format(LOCTEXT("AddSlotOverlayMaterialIndexLabelFormat", "Overlay Element {0}"), FText::AsNumber(SlotIndex));
+						FText AddSlotOverlayToolTip = !SlotName.IsNone() ?
+							FText::Format(LOCTEXT("AddSlotOverlayMaterialToolTipFormat", "Add per-slot overlay material for slot {0}, index {1}"), FText::FromName(SlotName), FText::AsNumber(SlotIndex)) :
+							FText::Format(LOCTEXT("AddSlotOverlayMaterialIndexToolTipFormat", "Add per-slot overlay material for element {0}"), FText::AsNumber(SlotIndex));
+						MenuBuilder.AddMenuEntry(AddSlotOverlayLabel, AddSlotOverlayToolTip, FSlateIcon(), AddComponentMaterialAction);
+					}
 				}
 			}
 			MenuBuilder.EndSection();
@@ -686,6 +717,11 @@ void FComponentMaterialTrackEditor::HandleAddComponentMaterialActionExecute(USce
 			case EComponentMaterialType::VolumetricCloudMaterial:
 					TrackDisplayName = FText::Format(LOCTEXT("CloudMaterialTrackName", "Volumetric Cloud Material {0}"), MaterialInterface ? FText::FromString(MaterialInterface->GetName()) : FText());
 					break;
+			case EComponentMaterialType::SlotOverlayMaterial:
+				TrackDisplayName = !MaterialInfo.MaterialSlotName.IsNone() ?
+					FText::Format(LOCTEXT("SlotOverlayMaterialTrackName", "Overlay Slot: {0}"), FText::FromName(MaterialInfo.MaterialSlotName)) :
+					FText::Format(LOCTEXT("IndexedSlotOverlayMaterialTrackName", "Overlay Element {0}"), FText::AsNumber(MaterialInfo.MaterialSlotIndex));
+				break;
 			default:
 				break;
 
diff --git a/Engine/Source/Editor/MovieSceneTools/Private/TrackEditors/PrimitiveMaterialTrackEditor.cpp b/Engine/Source/Editor/MovieSceneTools/Private/TrackEditors/PrimitiveMaterialTrackEditor.cpp
--- a/Engine/Source/Editor/MovieSceneTools/Private/TrackEditors/PrimitiveMaterialTrackEditor.cpp
+++ b/Engine/Source/Editor/MovieSceneTools/Private/TrackEditors/PrimitiveMaterialTrackEditor.cpp
@@ -99,6 +99,31 @@ void FPrimitiveMaterialTrackEditor::ConstructObjectBindingTrackMenu(FMenuBuilder
 						FText OverlayMaterialSwitcherTooltip = LOCTEXT("OverlayMaterialSwitcherTooltip_Format", "Add overlay material switcher");
 						MenuBuilder.AddMenuEntry(OverlayMaterialSwitcherLabel, OverlayMaterialSwitcherTooltip, FSlateIcon(), AddMaterialSwitcherAction);
 					}
+
+					const TArray<TObjectPtr<UMaterialInterface>>& SlotOverlayMaterials = MeshComponent->GetComponentMaterialSlotsOverlayMaterial();
+					for (int32 SlotIndex = 0; SlotIndex < SlotOverlayMaterials.Num(); ++SlotIndex)
+					{
+						UMaterialInterface* SlotOverlayMaterial = SlotOverlayMaterials[SlotIndex];
+						if (!SlotOverlayMaterial)
+						{
+							continue;
+						}
+						FName SlotName = MaterialSlotNames.IsValidIndex(SlotIndex) ? MaterialSlotNames[SlotIndex] : FName();
+						FComponentMaterialInfo SlotMaterialInfo{ SlotName, SlotIndex, EComponentMaterialType::SlotOverlayMaterial };
+						const bool bSlotAlreadyExists = Algo::FindBy(Binding->GetTracks(), SlotMaterialInfo, GetMaterialInfoForTrack) != nullptr;
+						if (bSlotAlreadyExists)
+						{
+							continue;
+						}
+						FUIAction AddSlotOverlaySwitcherAction(FExecuteAction::CreateSP(this, &FPrimitiveMaterialTrackEditor::CreateTrackForElement, ObjectBindings, SlotMaterialInfo));
+						FText SlotOverlaySwitcherLabel = !SlotName.IsNone() ?
+							FText::Format(LOCTEXT("SlotOverlayMaterialSwitcher_Format", "Overlay Slot {0} Switcher"), FText::FromName(SlotName)) :
+							FText::Format(LOCTEXT("IndexedSlotOverlayMaterialSwitcher_Format", "Overlay Element {0} Switcher"), FText::AsNumber(SlotIndex));
+						FText SlotOverlaySwitcherTooltip = !SlotName.IsNone() ?
+							FText::Format(LOCTEXT("SlotOverlayMaterialSwitcherTooltip_Format", "Add per-slot overlay material switcher for slot {0}, index {1}"), FText::FromName(SlotName), FText::AsNumber(SlotIndex)) :
+							FText::Format(LOCTEXT("IndexedSlotOverlayMaterialSwitcherTooltip_Format", "Add per-slot overlay material switcher for element {0}"), FText::AsNumber(SlotIndex));
+						MenuBuilder.AddMenuEntry(SlotOverlaySwitcherLabel, SlotOverlaySwitcherTooltip, FSlateIcon(), AddSlotOverlaySwitcherAction);
+					}
 				}
 			}
 			MenuBuilder.EndSection();
@@ -169,6 +194,11 @@ void FPrimitiveMaterialTrackEditor::CreateTrackForElement(TArray<FGuid> ObjectBi
 		case EComponentMaterialType::VolumetricCloudMaterial:
 				TrackDisplayName = LOCTEXT("CloudMaterialSwitcherTrackName", "Volumetric Cloud Material");
 				break;
+		case EComponentMaterialType::SlotOverlayMaterial:
+			TrackDisplayName = !MaterialInfo.MaterialSlotName.IsNone() ?
+				FText::Format(LOCTEXT("SlotOverlayMaterialSwitcherTrackName", "Overlay Slot: {0}"), FText::FromName(MaterialInfo.MaterialSlotName)) :
+				FText::Format(LOCTEXT("IndexedSlotOverlayMaterialSwitcherTrackName", "Overlay Element {0}"), FText::AsNumber(MaterialInfo.MaterialSlotIndex));
+			break;
 		default:
 			break;
 
diff --git a/Engine/Source/Runtime/MovieSceneTracks/Private/Systems/MovieSceneComponentMaterialSystem.cpp b/Engine/Source/Runtime/MovieSceneTracks/Private/Systems/MovieSceneComponentMaterialSystem.cpp
--- a/Engine/Source/Runtime/MovieSceneTracks/Private/Systems/MovieSceneComponentMaterialSystem.cpp
+++ b/Engine/Source/Runtime/MovieSceneTracks/Private/Systems/MovieSceneComponentMaterialSystem.cpp
@@ -83,6 +83,12 @@ UMaterialInterface* FComponentMaterialAccessor::GetMaterial() const
 			return CloudComponent->GetMaterial();
 		}
 		break;
+	case EComponentMaterialType::SlotOverlayMaterial:
+		if (UMeshComponent* Component = Cast<UMeshComponent>(Object))
+		{
+			return Component->GetOverlayMaterial(true, MaterialInfo.MaterialSlotIndex);
+		}
+		break;
 	default:
 		break;
 	}
@@ -127,6 +133,12 @@ void FComponentMaterialAccessor::SetMaterial(UMaterialInterface* InMaterial) con
 			CloudComponent->SetMaterial(InMaterial);
 		}
 		break;
+	case EComponentMaterialType::SlotOverlayMaterial:
+		if (UMeshComponent* Component = Cast<UMeshComponent>(Object))
+		{
+			Component->SetOverlayMaterial(InMaterial, true, MaterialInfo.MaterialSlotIndex);
+		}
+		break;
 	default:
 		break;
 	}
@@ -185,6 +197,14 @@ UMaterialInstanceDynamic* FComponentMaterialAccessor::CreateDynamicMaterial(UMat
 			return Result;
 		}
 		break;
+	case EComponentMaterialType::SlotOverlayMaterial:
+		if (UMeshComponent* Component = Cast<UMeshComponent>(Object))
+		{
+			UMaterialInstanceDynamic* Result = MakeDynamicMaterial();
+			Component->SetOverlayMaterial(Result, true, MaterialInfo.MaterialSlotIndex);
+			return Result;
+		}
+		break;
 	default:
 		break;
 	}
diff --git a/Engine/Source/Runtime/MovieSceneTracks/Private/Tracks/MovieSceneMaterialTrack.cpp b/Engine/Source/Runtime/MovieSceneTracks/Private/Tracks/MovieSceneMaterialTrack.cpp
--- a/Engine/Source/Runtime/MovieSceneTracks/Private/Tracks/MovieSceneMaterialTrack.cpp
+++ b/Engine/Source/Runtime/MovieSceneTracks/Private/Tracks/MovieSceneMaterialTrack.cpp
@@ -391,6 +391,24 @@ FText UMovieSceneComponentMaterialTrack::GetDisplayNameToolTipText(const FMovieS
 						}
 					}
 					break;
+			case EComponentMaterialType::SlotOverlayMaterial:
+				if (UMeshComponent* Component = Cast<UMeshComponent>(WeakObject.Get()))
+				{
+					Material = Component->GetOverlayMaterial(true, MaterialInfo.MaterialSlotIndex);
+					if (Material)
+					{
+						return !MaterialInfo.MaterialSlotName.IsNone() ?
+							FText::Format(LOCTEXT("SlotOverlayMaterialTrackTooltip", "Material parameter track for overlay slot {0} at index {1}, asset {2}"), FText::FromName(MaterialInfo.MaterialSlotName), FText::AsNumber(MaterialInfo.MaterialSlotIndex), FText::FromString(Material->GetName())) :
+							FText::Format(LOCTEXT("IndexedSlotOverlayMaterialTrackTooltip", "Material parameter track for overlay element at index {0}, asset {1}"), FText::AsNumber(MaterialInfo.MaterialSlotIndex), FText::FromString(Material->GetName()));
+					}
+					else
+					{
+						return !MaterialInfo.MaterialSlotName.IsNone() ?
+							FText::Format(LOCTEXT("SlotOverlayMaterialTrackTooltip_Missing", "No overlay material found for slot {0} at index {1}"), FText::FromName(MaterialInfo.MaterialSlotName), FText::AsNumber(MaterialInfo.MaterialSlotIndex)) :
+							FText::Format(LOCTEXT("IndexedSlotOverlayMaterialTrackTooltip_Missing", "No overlay material found at index {0}"), FText::AsNumber(MaterialInfo.MaterialSlotIndex));
+					}
+				}
+				break;
 			default:
 				break;
 			}
@@ -502,6 +520,20 @@ FSlateColor UMovieSceneComponentMaterialTrack::GetLabelColor(const FMovieSceneLa
 					return GetDimmedColor(FLinearColor::Red);
 				}
 				break;
+			case EComponentMaterialType::SlotOverlayMaterial:
+				if (UMeshComponent* Component = Cast<UMeshComponent>(WeakObject.Get()))
+				{
+					UMaterialInterface* Material = Component->GetOverlayMaterial(true, MaterialInfo.MaterialSlotIndex);
+					if (!Material)
+					{
+						return GetDimmedColor(FLinearColor::Red);
+					}
+				}
+				else
+				{
+					return GetDimmedColor(FLinearColor::Red);
+				}
+				break;
 			default:
 				break;
 			}
@@ -522,4 +554,4 @@ void UMovieSceneComponentMaterialTrack::PostLoad()
 }
 #endif
 
-#undef LOCTEXT_NAMESPACE
\ No newline at end of file
+#undef LOCTEXT_NAMESPACE
diff --git a/Engine/Source/Runtime/MovieSceneTracks/Private/Tracks/MovieScenePrimitiveMaterialTrack.cpp b/Engine/Source/Runtime/MovieSceneTracks/Private/Tracks/MovieScenePrimitiveMaterialTrack.cpp
--- a/Engine/Source/Runtime/MovieSceneTracks/Private/Tracks/MovieScenePrimitiveMaterialTrack.cpp
+++ b/Engine/Source/Runtime/MovieSceneTracks/Private/Tracks/MovieScenePrimitiveMaterialTrack.cpp
@@ -90,6 +90,14 @@ FText UMovieScenePrimitiveMaterialTrack::GetDisplayNameToolTipText(const FMovieS
 					return LOCTEXT("DecalMaterialSwitcherTrackTooltip", "Material switcher for decal material");
 				}
 				break;
+			case EComponentMaterialType::SlotOverlayMaterial:
+				if (UMeshComponent* Component = Cast<UMeshComponent>(WeakObject.Get()))
+				{
+					return !MaterialInfo.MaterialSlotName.IsNone() ?
+						FText::Format(LOCTEXT("SlotOverlayMaterialSwitcherTooltip", "Material switcher for overlay slot {0} at index {1}"), FText::FromName(MaterialInfo.MaterialSlotName), FText::AsNumber(MaterialInfo.MaterialSlotIndex)) :
+						FText::Format(LOCTEXT("IndexedSlotOverlayMaterialSwitcherTooltip", "Material switcher for overlay element at index {0}"), FText::AsNumber(MaterialInfo.MaterialSlotIndex));
+				}
+				break;
 			default:
 				break;
 			}
@@ -177,6 +185,15 @@ FSlateColor UMovieScenePrimitiveMaterialTrack::GetLabelColor(const FMovieSceneLa
 				}
 			}
 				break;
+			case EComponentMaterialType::SlotOverlayMaterial:
+			{
+				UMeshComponent* MeshComponent = Cast<UMeshComponent>(WeakObject.Get());
+				if (!MeshComponent)
+				{
+					return GetDimmedColor(FLinearColor::Red);
+				}
+			}
+				break;
 			default:
 				break;
 			}
@@ -219,4 +236,4 @@ void UMovieScenePrimitiveMaterialTrack::SetMaterialIndex(int32 InMaterialIndex)
 	MaterialInfo.MaterialType = EComponentMaterialType::IndexedMaterial;
 }
 
-#undef LOCTEXT_NAMESPACE
\ No newline at end of file
+#undef LOCTEXT_NAMESPACE
diff --git a/Engine/Source/Runtime/MovieSceneTracks/Public/Tracks/MovieSceneMaterialTrack.h b/Engine/Source/Runtime/MovieSceneTracks/Public/Tracks/MovieSceneMaterialTrack.h
--- a/Engine/Source/Runtime/MovieSceneTracks/Public/Tracks/MovieSceneMaterialTrack.h
+++ b/Engine/Source/Runtime/MovieSceneTracks/Public/Tracks/MovieSceneMaterialTrack.h
@@ -25,7 +25,9 @@ enum class EComponentMaterialType : uint8
 	/* A decal material*/
 	DecalMaterial,
 	/* Volumetric Cloud Material*/
-	VolumetricCloudMaterial
+	VolumetricCloudMaterial,
+	/* A per-slot overlay material on a mesh component */
+	SlotOverlayMaterial,
 };
 
 /**
@@ -74,6 +76,16 @@ struct FComponentMaterialInfo
 		case EComponentMaterialType::DecalMaterial:
 			return TEXT("Decal Material");
 			break;
+		case EComponentMaterialType::SlotOverlayMaterial:
+			if (!MaterialSlotName.IsNone())
+			{
+				return FString::Printf(TEXT("Overlay Material Slot: %s"), *MaterialSlotName.ToString());
+			}
+			else
+			{
+				return FString::Printf(TEXT("Overlay Material Element %i"), MaterialSlotIndex);
+			}
+			break;
 		}
 		return FString();
 	}
