CurveEditor: Fix an issue where curves that use FChannelCurveModel::DrawCurve in some cases failed to draw the curve line before the first visible key correctly

#Synced-CL 51029743

diff --git a/Engine/Source/Editor/MovieSceneTools/Private/Channels/ChannelCurveModel.cpp b/Engine/Source/Editor/MovieSceneTools/Private/Channels/ChannelCurveModel.cpp
--- a/Engine/Source/Editor/MovieSceneTools/Private/Channels/ChannelCurveModel.cpp
+++ b/Engine/Source/Editor/MovieSceneTools/Private/Channels/ChannelCurveModel.cpp
@@ -136,8 +136,18 @@ void FChannelCurveModel<ChannelType, ChannelValue, KeyType>::DrawCurve(const FCu
 		const FFrameNumber StartFrame = (StartTimeSeconds * TickResolution).FloorToFrame();
 		const FFrameNumber EndFrame = (EndTimeSeconds * TickResolution).CeilToFrame();
 
-		const int32 StartingIndex = Algo::UpperBound(Times, StartFrame);
-		const int32 EndingIndex = Algo::LowerBound(Times, EndFrame);
+		const int32 StartingIndex = [&Times, &StartFrame]()
+			{
+				const int32 StartingFrameIndex = Algo::LowerBound(Times, StartFrame);
+				if (Times.IsValidIndex(StartingFrameIndex - 1))
+				{
+					return StartingFrameIndex - 1;
+				}
+				
+				return StartingFrameIndex;
+			}();
+
+		const int32 EndingIndex = Algo::UpperBound(Times, EndFrame);
 
 		// Add the lower bound of the visible space
 		const bool bValidRange = StartingIndex < EndingIndex;
