Visual Cue in Sequencer Section where loop state and/or duration doesn?t match the underlying asset
UE-349249

Showing state by colour led to various threads about how the UI is about to be rewritten and to please hold off on that.
Thus, for 5.8, I've reworked the change to communicating via the tooltip, leaving discussions for how to most effectively communicate section/track state by colour, tooltips, overlays, etc. for the PRD and likely UE6.

#Synced-CL 50895136

diff --git a/Engine/Source/Editor/MovieSceneTools/Private/TrackEditors/AudioTrackEditor.cpp b/Engine/Source/Editor/MovieSceneTools/Private/TrackEditors/AudioTrackEditor.cpp
--- a/Engine/Source/Editor/MovieSceneTools/Private/TrackEditors/AudioTrackEditor.cpp
+++ b/Engine/Source/Editor/MovieSceneTools/Private/TrackEditors/AudioTrackEditor.cpp
@@ -54,6 +54,8 @@
 
 #define LOCTEXT_NAMESPACE "FAudioTrackEditor"
 
+// Tolerance for Section vs. Sound duration to be considered nearly equal.
+constexpr float SectionToSoundDurationToleranceInSeconds = 0.1f;
 
 namespace AnimatableAudioEditorConstants
 {
@@ -923,13 +925,26 @@ FText FAudioSection::GetSectionToolTip() const
 		}
 
 		FText LoopStateDivergenceToolTip;
-
-		if (AudioSection->GetLooping() == AudioSection->GetSound()->IsOneShot())
+		if (AudioSection->GetLooping() == Sound->IsOneShot())
 		{
+			// Section loop state does not match the underlying Sound.
 			LoopStateDivergenceToolTip = LOCTEXT("ToolTipLoopStateDivergence", "\nOverriding loop state from source content");
 		}
 
-		ToolTip = FText::Format(LOCTEXT("FullToolTipContentFormat", "{0}{1}{2}"), BaseToolTip, LoopStateDivergenceToolTip, CompletionModeToolTip);
+		FText DurationDivergenceToolTip;
+		check(AudioSection->HasStartFrame());
+		check(AudioSection->HasEndFrame());
+		if (!AudioSection->GetLooping() && !Sound->IsProcedurallyGenerated())
+		{
+			const float SoundDurationInSeconds = DeriveUnloopedDuration(AudioSection);
+			if (!FMath::IsNearlyEqual(SectionDuration, SoundDurationInSeconds, SectionToSoundDurationToleranceInSeconds))
+			{
+				// Section duration does not match the underlying Sound.
+				DurationDivergenceToolTip = LOCTEXT("ToolTipDurationDivergence", "\nOverriding duration from source content");
+			}
+		}
+
+		ToolTip = FText::Format(LOCTEXT("FullToolTipContentFormat", "{0}{1}{2}{3}"), BaseToolTip, LoopStateDivergenceToolTip, DurationDivergenceToolTip, CompletionModeToolTip);
 	}
 
 	ToolTip = FText::Join(FText::GetEmpty(), ToolTip, AttachToolTip);
