Anim In Engine: Add option to not record empty zero curves

#Synced-CL 50883889

diff --git a/Engine/Plugins/Animation/ControlRig/Source/ControlRigEditor/Private/ControlRigEditorModule.cpp b/Engine/Plugins/Animation/ControlRig/Source/ControlRigEditor/Private/ControlRigEditorModule.cpp
--- a/Engine/Plugins/Animation/ControlRig/Source/ControlRigEditor/Private/ControlRigEditorModule.cpp
+++ b/Engine/Plugins/Animation/ControlRig/Source/ControlRigEditor/Private/ControlRigEditorModule.cpp
@@ -1426,6 +1426,7 @@ void FControlRigEditorModule::BakeToControlRig(UClass* ControlRigClass, UAnimSeq
 							LevelAnimLinkItem.Interpolation = EAnimInterpolationType::Linear;
 							LevelAnimLinkItem.CurveInterpolation = ERichCurveInterpMode::RCIM_Linear;
 							LevelAnimLinkItem.bExportMaterialCurves = true;
+							LevelAnimLinkItem.bSkipCurvesWithZeroValue = true;
 							LevelAnimLinkItem.bExportTransforms = true;
 							LevelAnimLinkItem.bRecordInWorldSpace = false;
 							LevelAnimLinkItem.bEvaluateAllSkeletalMeshComponents = true;
diff --git a/Engine/Plugins/MovieScene/SequencerScripting/Source/SequencerScriptingEditor/Private/SequencerTools.cpp b/Engine/Plugins/MovieScene/SequencerScripting/Source/SequencerScriptingEditor/Private/SequencerTools.cpp
--- a/Engine/Plugins/MovieScene/SequencerScripting/Source/SequencerScriptingEditor/Private/SequencerTools.cpp
+++ b/Engine/Plugins/MovieScene/SequencerScripting/Source/SequencerScriptingEditor/Private/SequencerTools.cpp
@@ -556,6 +556,7 @@ bool USequencerToolsFunctionLibrary::LinkAnimSequence(ULevelSequence*  Sequence,
 						LevelAnimLinkItem.bExportMorphTargets = ExportOptions->bExportMorphTargets;
 						LevelAnimLinkItem.bExportAttributeCurves = ExportOptions->bExportAttributeCurves;
 						LevelAnimLinkItem.bExportMaterialCurves = ExportOptions->bExportMaterialCurves;
+						LevelAnimLinkItem.bSkipCurvesWithZeroValue = ExportOptions->bSkipCurvesWithZeroValue;
 						LevelAnimLinkItem.bExportTransforms = ExportOptions->bExportTransforms;
 						LevelAnimLinkItem.Interpolation = ExportOptions->Interpolation;
 						LevelAnimLinkItem.CurveInterpolation = ExportOptions->CurveInterpolation;
@@ -593,6 +594,7 @@ bool USequencerToolsFunctionLibrary::LinkAnimSequence(ULevelSequence*  Sequence,
 				LevelAnimLinkItem.Interpolation = ExportOptions->Interpolation;
 				LevelAnimLinkItem.CurveInterpolation = ExportOptions->CurveInterpolation;
 				LevelAnimLinkItem.bExportMaterialCurves = ExportOptions->bExportMaterialCurves;
+				LevelAnimLinkItem.bSkipCurvesWithZeroValue = ExportOptions->bSkipCurvesWithZeroValue;
 				LevelAnimLinkItem.bExportTransforms = ExportOptions->bExportTransforms;
 				LevelAnimLinkItem.bRecordInWorldSpace = ExportOptions->bRecordInWorldSpace;
 				LevelAnimLinkItem.bEvaluateAllSkeletalMeshComponents = ExportOptions->bEvaluateAllSkeletalMeshComponents;
diff --git a/Engine/Source/Editor/MovieSceneTools/Private/MovieSceneToolHelpers.cpp b/Engine/Source/Editor/MovieSceneTools/Private/MovieSceneToolHelpers.cpp
--- a/Engine/Source/Editor/MovieSceneTools/Private/MovieSceneToolHelpers.cpp
+++ b/Engine/Source/Editor/MovieSceneTools/Private/MovieSceneToolHelpers.cpp
@@ -4464,6 +4464,7 @@ bool MovieSceneToolHelpers::ExportToAnimSequence(UAnimSequence* AnimSequence, UA
 			RecordingSettings.bRecordMorphTargets = ExportOptions->bExportMorphTargets;
 			RecordingSettings.bRecordAttributeCurves = ExportOptions->bExportAttributeCurves;
 			RecordingSettings.bRecordMaterialCurves = ExportOptions->bExportMaterialCurves;
+			RecordingSettings.bSkipCurvesWithZeroValue = ExportOptions->bSkipCurvesWithZeroValue;
 			RecordingSettings.bRecordInWorldSpace = ExportOptions->bRecordInWorldSpace;
 			RecordingSettings.IncludeAnimationNames = ExportOptions->IncludeAnimationNames;
 			RecordingSettings.ExcludeAnimationNames = ExportOptions->ExcludeAnimationNames;
diff --git a/Engine/Source/Editor/MovieSceneTools/Private/TrackEditors/CommonAnimationTrackEditor.cpp b/Engine/Source/Editor/MovieSceneTools/Private/TrackEditors/CommonAnimationTrackEditor.cpp
--- a/Engine/Source/Editor/MovieSceneTools/Private/TrackEditors/CommonAnimationTrackEditor.cpp
+++ b/Engine/Source/Editor/MovieSceneTools/Private/TrackEditors/CommonAnimationTrackEditor.cpp
@@ -1546,6 +1546,7 @@ void FCommonAnimationTrackEditor::OnSequencerSaved(ISequencer& )
 							const bool bSavedExportMorphTargets = AnimSeqExportOption->bExportMorphTargets;
 							const bool bSavedExportAttributeCurves = AnimSeqExportOption->bExportAttributeCurves;
 							const bool bSavedExportMaterialCurves = AnimSeqExportOption->bExportMaterialCurves;
+							const bool bSkipCurvesWithZeroValue = AnimSeqExportOption->bSkipCurvesWithZeroValue;
 							const bool bSavedExportTransforms = AnimSeqExportOption->bExportTransforms;
 							const bool bSavedIncludeComponentTransform = AnimSeqExportOption->bRecordInWorldSpace;
 							const bool bSavedEvaluateAllSkeletalMeshComponents = AnimSeqExportOption->bEvaluateAllSkeletalMeshComponents;
@@ -1581,6 +1582,7 @@ void FCommonAnimationTrackEditor::OnSequencerSaved(ISequencer& )
 							AnimSeqExportOption->bExportMorphTargets = Item.bExportMorphTargets;
 							AnimSeqExportOption->bExportAttributeCurves = Item.bExportAttributeCurves;
 							AnimSeqExportOption->bExportMaterialCurves = Item.bExportMaterialCurves;
+							AnimSeqExportOption->bSkipCurvesWithZeroValue = Item.bSkipCurvesWithZeroValue;
 							AnimSeqExportOption->bExportTransforms = Item.bExportTransforms;
 							AnimSeqExportOption->bRecordInWorldSpace = Item.bRecordInWorldSpace;
 							AnimSeqExportOption->bEvaluateAllSkeletalMeshComponents = Item.bEvaluateAllSkeletalMeshComponents;
@@ -1613,6 +1615,7 @@ void FCommonAnimationTrackEditor::OnSequencerSaved(ISequencer& )
 							AnimSeqExportOption->bExportMorphTargets = bSavedExportMorphTargets;
 							AnimSeqExportOption->bExportAttributeCurves = bSavedExportAttributeCurves;
 							AnimSeqExportOption->bExportMaterialCurves = bSavedExportMaterialCurves;
+							AnimSeqExportOption->bSkipCurvesWithZeroValue = bSkipCurvesWithZeroValue;
 							AnimSeqExportOption->bExportTransforms = bSavedExportTransforms;
 							AnimSeqExportOption->bRecordInWorldSpace = bSavedIncludeComponentTransform;
 							AnimSeqExportOption->bEvaluateAllSkeletalMeshComponents = bSavedEvaluateAllSkeletalMeshComponents;
@@ -1798,6 +1801,7 @@ bool FCommonAnimationTrackEditor::CreateAnimationSequence(const TArray<UObject*>
 								LevelAnimLinkItem.bExportMorphTargets = AnimSeqExportOption->bExportMorphTargets;
 								LevelAnimLinkItem.bExportAttributeCurves = AnimSeqExportOption->bExportAttributeCurves;
 								LevelAnimLinkItem.bExportMaterialCurves = AnimSeqExportOption->bExportMaterialCurves;
+								LevelAnimLinkItem.bSkipCurvesWithZeroValue = AnimSeqExportOption->bSkipCurvesWithZeroValue;
 								LevelAnimLinkItem.bExportTransforms = AnimSeqExportOption->bExportTransforms;
 								LevelAnimLinkItem.bRecordInWorldSpace = AnimSeqExportOption->bRecordInWorldSpace;
 								LevelAnimLinkItem.bEvaluateAllSkeletalMeshComponents = AnimSeqExportOption->bEvaluateAllSkeletalMeshComponents;
@@ -1831,6 +1835,7 @@ bool FCommonAnimationTrackEditor::CreateAnimationSequence(const TArray<UObject*>
 						LevelAnimLinkItem.bExportMorphTargets = AnimSeqExportOption->bExportMorphTargets;
 						LevelAnimLinkItem.bExportAttributeCurves = AnimSeqExportOption->bExportAttributeCurves;
 						LevelAnimLinkItem.bExportMaterialCurves = AnimSeqExportOption->bExportMaterialCurves;
+						LevelAnimLinkItem.bSkipCurvesWithZeroValue = AnimSeqExportOption->bSkipCurvesWithZeroValue;
 						LevelAnimLinkItem.bExportTransforms = AnimSeqExportOption->bExportTransforms;
 						LevelAnimLinkItem.bRecordInWorldSpace = AnimSeqExportOption->bRecordInWorldSpace;
 						LevelAnimLinkItem.bEvaluateAllSkeletalMeshComponents = AnimSeqExportOption->bEvaluateAllSkeletalMeshComponents;
diff --git a/Engine/Source/Editor/SequenceRecorder/Private/AnimationRecorder.cpp b/Engine/Source/Editor/SequenceRecorder/Private/AnimationRecorder.cpp
--- a/Engine/Source/Editor/SequenceRecorder/Private/AnimationRecorder.cpp
+++ b/Engine/Source/Editor/SequenceRecorder/Private/AnimationRecorder.cpp
@@ -470,17 +470,17 @@ UAnimSequence* FAnimationRecorder::StopRecord(bool bShowMessage)
 					const bool bMaterial = EnumHasAnyFlags(InElement.Flags, UE::Anim::ECurveElementFlags::Material);
 					const bool bShouldSkipName = ShouldSkipName(InElement.Name);
 					const bool bAttributeCurve = !bMorphTarget && !bMaterial;
-					const bool bSkipCurve = (bMorphTarget && !bRecordMorphTargets) ||
+					bool bSkipCurve = (bMorphTarget && !bRecordMorphTargets) ||
 											(bAttributeCurve && !bRecordAttributeCurves) ||
 											(bMaterial && !bRecordMaterialCurves) ||
 											bShouldSkipName;
-
 					if(!bSkipCurve)
 					{
 						AllCurveNames.Add(InElement.Name);
 					}
 				});
 			}
+			
 
 			StartTime = FPlatformTime::Seconds();
 			USkeleton* SkeletonObj = AnimationObject->GetSkeleton();
@@ -526,17 +526,34 @@ UAnimSequence* FAnimationRecorder::StopRecord(bool bShowMessage)
 				// Fill all the curve data at once
 				if (FloatCurveData)
 				{
+					//since default is true to check nonzero values we avoid the if state to test and just
+					//always check
+					bool bHasNonZeroValue = false;
 					TArray<FRichCurveKey> Keys;
 					for (int32 Index = 0; Index < WriteIndex; ++Index)
 					{
-						FRichCurveKey Key(TimesToRecord[Index], ValuesToRecord[Index]);
+						float Value = ValuesToRecord[Index];
+						if (bHasNonZeroValue == false)
+						{
+							if (FMath::IsNearlyZero(Value, UE_KINDA_SMALL_NUMBER) == false)
+							{
+								bHasNonZeroValue = true;
+							}
+						}
+						FRichCurveKey Key(TimesToRecord[Index], Value);
 						Key.InterpMode = InterpMode;
 						Key.TangentMode = TangentMode;
 						Keys.Add(Key);
 					}
-
 					const FAnimationCurveIdentifier CurveId(FloatCurveData->GetName(), ERawCurveTrackTypes::RCT_Float);
-					Controller.SetCurveKeys(CurveId, Keys, bTransactRecording);
+					if(bSkipCurvesWithZeroValue == false || bHasNonZeroValue)
+					{ 
+						Controller.SetCurveKeys(CurveId, Keys, bTransactRecording);
+					}
+					else
+					{
+						Controller.RemoveCurve(CurveId, bTransactRecording);
+					}
 				}
 			}
 
@@ -1214,6 +1231,7 @@ void FAnimRecorderInstance::InitInternal(USkeletalMeshComponent* InComponent, co
 	Recorder->bRecordMorphTargets = Settings.bRecordMorphTargets;
 	Recorder->bRecordAttributeCurves = Settings.bRecordAttributeCurves;
 	Recorder->bRecordMaterialCurves = Settings.bRecordMaterialCurves;
+	Recorder->bSkipCurvesWithZeroValue = Settings.bSkipCurvesWithZeroValue;
 	Recorder->bTransactRecording = Settings.bTransactRecording;
 	Recorder->IncludeAnimationNames = Settings.IncludeAnimationNames;
 	Recorder->ExcludeAnimationNames = Settings.ExcludeAnimationNames;
diff --git a/Engine/Source/Editor/SequenceRecorder/Public/AnimationRecorder.h b/Engine/Source/Editor/SequenceRecorder/Public/AnimationRecorder.h
--- a/Engine/Source/Editor/SequenceRecorder/Public/AnimationRecorder.h
+++ b/Engine/Source/Editor/SequenceRecorder/Public/AnimationRecorder.h
@@ -196,6 +196,8 @@ public:
 	uint8 bRecordAttributeCurves : 1;
 	/** Whether or not to record material curves*/
 	uint8 bRecordMaterialCurves : 1;
+	/** If enabled, we will skip any curves that have only zero values */
+	uint8 bSkipCurvesWithZeroValue : 1;
 	/** Include list */
 	TArray<FString> IncludeAnimationNames;
 	/** Exclude list */
diff --git a/Engine/Source/Editor/UnrealEd/Classes/Exporters/AnimSeqExportOption.h b/Engine/Source/Editor/UnrealEd/Classes/Exporters/AnimSeqExportOption.h
--- a/Engine/Source/Editor/UnrealEd/Classes/Exporters/AnimSeqExportOption.h
+++ b/Engine/Source/Editor/UnrealEd/Classes/Exporters/AnimSeqExportOption.h
@@ -33,6 +33,10 @@ public:
 	UPROPERTY(EditAnywhere, BlueprintReadWrite, AdvancedDisplay, Category = Export)
 	bool bExportMaterialCurves = true;
 
+	/** If enabled, we will skip any curves that have only zero values */
+	UPROPERTY(EditAnywhere, BlueprintReadWrite, AdvancedDisplay, Category = Export)
+	bool bSkipCurvesWithZeroValue = true;
+
 	/** If enabled we record in World Space otherwise we record from 0,0,0 */
 	UPROPERTY(EditAnywhere, BlueprintReadWrite, AdvancedDisplay, Category = Export)
 	bool bRecordInWorldSpace = false;
diff --git a/Engine/Source/Editor/UnrealEd/Private/Animation/AnimSeqExportOption.cpp b/Engine/Source/Editor/UnrealEd/Private/Animation/AnimSeqExportOption.cpp
--- a/Engine/Source/Editor/UnrealEd/Private/Animation/AnimSeqExportOption.cpp
+++ b/Engine/Source/Editor/UnrealEd/Private/Animation/AnimSeqExportOption.cpp
@@ -21,6 +21,7 @@ void UAnimSeqExportOption::ResetToDefault()
 	bExportMorphTargets = true;
 	bExportAttributeCurves = true;
 	bExportMaterialCurves = true;
+	bSkipCurvesWithZeroValue = true;
 	bRecordInWorldSpace = false;
 	bSetRetargetSourceAsset = CVarSetRetargetSourceAsset.GetValueOnGameThread();
 	Interpolation = EAnimInterpolationType::Linear;
diff --git a/Engine/Source/Runtime/Engine/Classes/Animation/AnimationRecordingSettings.h b/Engine/Source/Runtime/Engine/Classes/Animation/AnimationRecordingSettings.h
--- a/Engine/Source/Runtime/Engine/Classes/Animation/AnimationRecordingSettings.h
+++ b/Engine/Source/Runtime/Engine/Classes/Animation/AnimationRecordingSettings.h
@@ -39,6 +39,7 @@ struct FAnimationRecordingSettings
 		, bRecordMorphTargets(true)
 		, bRecordAttributeCurves(true)
 		, bRecordMaterialCurves(true)
+		, bSkipCurvesWithZeroValue(true)
 		, bTransactRecording(true) 
 	{}
 	
@@ -98,6 +99,10 @@ struct FAnimationRecordingSettings
 	UPROPERTY(EditAnywhere, Category = "Settings")
 	bool bRecordMaterialCurves;
 
+	/** If enabled, we will skip any curves that have only zero values */
+	UPROPERTY(EditAnywhere, Category = "Settings")
+	bool bSkipCurvesWithZeroValue;
+
 	/** Whether or not to transact recording changes */
 	UPROPERTY(EditAnywhere, Category = "Settings")
 	bool bTransactRecording;
diff --git a/Engine/Source/Runtime/LevelSequence/Public/LevelSequenceAnimSequenceLink.h b/Engine/Source/Runtime/LevelSequence/Public/LevelSequenceAnimSequenceLink.h
--- a/Engine/Source/Runtime/LevelSequence/Public/LevelSequenceAnimSequenceLink.h
+++ b/Engine/Source/Runtime/LevelSequence/Public/LevelSequenceAnimSequenceLink.h
@@ -38,6 +38,8 @@ public:
 	UPROPERTY(BlueprintReadWrite, Category = Property)
 	bool bExportMaterialCurves = true;
 	UPROPERTY(BlueprintReadWrite, Category = Property)
+	bool bSkipCurvesWithZeroValue = false;
+	UPROPERTY(BlueprintReadWrite, Category = Property)
 	bool bSetRetargetSourceAsset = true;
 	UPROPERTY(BlueprintReadWrite, Category = Property);
 	EAnimInterpolationType Interpolation = EAnimInterpolationType::Linear;
