Sequencer: Use display text for the perlin noise menu instead of the channel name. The display text will have been formatted through AxisDisplayInfo::GetAxisDisplayName

#Synced-CL 51424575

diff --git a/Engine/Source/Editor/MovieSceneTools/Private/Channels/PerlinNoiseChannelInterface.cpp b/Engine/Source/Editor/MovieSceneTools/Private/Channels/PerlinNoiseChannelInterface.cpp
--- a/Engine/Source/Editor/MovieSceneTools/Private/Channels/PerlinNoiseChannelInterface.cpp
+++ b/Engine/Source/Editor/MovieSceneTools/Private/Channels/PerlinNoiseChannelInterface.cpp
@@ -116,11 +116,25 @@ void FPerlinNoiseChannelSectionMenuExtension::BuildChannelsMenu(FMenuBuilder& Me
 	{
 		const int32 SectionIndex(ChannelHandleSectionIndexes[Index]);
 		const FMovieSceneChannelHandle ChannelHandle(ChannelHandles[Index]);
+		const FMovieSceneChannelMetaData* MetaData = ChannelHandle.GetMetaData();
+
+		FText DisplayText;
+		if (MetaData)
+		{
+			if (!MetaData->Group.IsEmpty())
+			{
+				DisplayText = FText::Format(LOCTEXT("PerlinNoiseChannelDisplayText","{0}.{1}"), MetaData->Group, MetaData->DisplayText);
+			}
+			else
+			{
+				DisplayText = MetaData->DisplayText;
+			}
+		}
 
 		if (bMultipleSections)
 		{
 			MenuBuilder.AddSubMenu(
-				FText::Format(LOCTEXT("PerlinNoiseChannelAndSectionSelectMenu", "Section{0}.{1}"), SectionIndex + 1, FText::FromName(ChannelHandle.GetMetaData()->Name)),
+				FText::Format(LOCTEXT("PerlinNoiseChannelAndSectionSelectMenu", "Section{0}.{1}"), SectionIndex + 1, DisplayText),
 				LOCTEXT("PerlinNoiseChannelAndSectionSelectMenuToolTip", "Edit parameters for this Perlin Noise channel"),
 				FNewMenuDelegate::CreateLambda([this, Index](FMenuBuilder& InnerMenuBuilder)
 					{
@@ -131,7 +145,7 @@ void FPerlinNoiseChannelSectionMenuExtension::BuildChannelsMenu(FMenuBuilder& Me
 		else
 		{
 			MenuBuilder.AddSubMenu(
-				FText::FromName(ChannelHandle.GetMetaData()->Name),
+				DisplayText,
 				LOCTEXT("PerlinNoiseChannelSelectMenuToolTip", "Edit parameters for this Perlin Noise channel"),
 				FNewMenuDelegate::CreateLambda([this, Index](FMenuBuilder& InnerMenuBuilder)
 					{
