Sequencer: Fix bake transform so that it gathers all the key times of all the transform track sections, rather than just the first section

This fixes baking with keys only for multiple sections

#Synced-CL 51597275

diff --git a/Engine/Source/Editor/MovieSceneTools/Private/MovieSceneToolHelpers.cpp b/Engine/Source/Editor/MovieSceneTools/Private/MovieSceneToolHelpers.cpp
--- a/Engine/Source/Editor/MovieSceneTools/Private/MovieSceneToolHelpers.cpp
+++ b/Engine/Source/Editor/MovieSceneTools/Private/MovieSceneToolHelpers.cpp
@@ -5743,6 +5743,13 @@ void MovieSceneToolHelpers::GetActorParentsWithAttachments(ISequencer* Sequencer
 void MovieSceneToolHelpers::GetActorsAndParentsKeyFrames(ISequencer* Sequencer, const FActorForWorldTransforms& InActor,
 	const FFrameNumber& StartFrame, const FFrameNumber& EndFrame, TSortedMap<FFrameNumber, FFrameNumber>& OutFrameMap)
 {
+	UMovieSceneSequence* Sequence = Sequencer ? Sequencer->GetFocusedMovieSceneSequence() : nullptr;
+	UMovieScene* MovieScene = Sequence ? Sequence->GetMovieScene() : nullptr;
+	if (!MovieScene)
+	{
+		return;
+	}
+
 	TArray<FFrameNumber> FramesToUse;
 	TArray<FActorForWorldTransforms> ParentActors;
 	ParentActors.Add(InActor);
@@ -5761,14 +5768,23 @@ void MovieSceneToolHelpers::GetActorsAndParentsKeyFrames(ISequencer* Sequencer,
 		}
 		if (Guid.IsValid())
 		{
-			if (UMovieScene3DTransformSection* TransformSection = MovieSceneToolHelpers::GetTransformSection(Sequencer, Guid))
-			{
-				TArrayView<FMovieSceneDoubleChannel*> Channels = TransformSection->GetChannelProxy().GetChannels<FMovieSceneDoubleChannel>();
-				TArray<FFrameNumber> TransformFrameTimes = FMovieSceneConstraintChannelHelper::GetTransformTimes(
-					Channels, StartFrame, EndFrame);
-				for (FFrameNumber& FrameNumber : TransformFrameTimes)
+			for (UMovieSceneTrack* Track : MovieScene->FindTracks(UMovieScene3DTransformTrack::StaticClass(), Guid))
+			{	
+				if (Track)
 				{
-					OutFrameMap.Add(FrameNumber,FrameNumber);
+					for (UMovieSceneSection* Section : Track->GetAllSections())
+					{
+						if (Section)
+						{
+							TArrayView<FMovieSceneDoubleChannel*> Channels = Section->GetChannelProxy().GetChannels<FMovieSceneDoubleChannel>();
+							TArray<FFrameNumber> TransformFrameTimes = FMovieSceneConstraintChannelHelper::GetTransformTimes(
+								Channels, StartFrame, EndFrame);
+							for (FFrameNumber& FrameNumber : TransformFrameTimes)
+							{
+								OutFrameMap.Add(FrameNumber, FrameNumber);
+							}
+						}
+					}
 				}
 			}
 		}
