Sequencer: When Force Front XAxis is enabled, tangents should be negated.

#Synced-CL 51953805

diff --git a/Engine/Source/Editor/MovieSceneTools/Private/MovieSceneToolHelpers.cpp b/Engine/Source/Editor/MovieSceneTools/Private/MovieSceneToolHelpers.cpp
--- a/Engine/Source/Editor/MovieSceneTools/Private/MovieSceneToolHelpers.cpp
+++ b/Engine/Source/Editor/MovieSceneTools/Private/MovieSceneToolHelpers.cpp
@@ -3157,8 +3157,9 @@ bool ImportFBXTransform(FString NodeName, FGuid ObjectBinding, UnFbx::FFbxCurves
 	}
 	if (DoubleChannels[4].Get())
 	{
+		const bool bNegateTangents = ImportFBXSettings->bForceFrontXAxis;
 		DoubleChannels[4].Get()->SetDefault(Rotation.Y);
-		ImportTransformChannelToDouble(EulerRotation[1], DoubleChannels[4].Get(), FrameRate, false, true);
+		ImportTransformChannelToDouble(EulerRotation[1], DoubleChannels[4].Get(), FrameRate, bNegateTangents, true);
 	}
 	if (DoubleChannels[5].Get())
 	{
