Sequencer: Fix SequencerKeyEditor to Modify() the section when the channel owner differs from the section. This is needed for decorations with channels.

#Synced-CL 51964473

diff --git a/Engine/Source/Editor/MovieSceneTools/Private/CurveKeyEditors/SequencerKeyEditor.h b/Engine/Source/Editor/MovieSceneTools/Private/CurveKeyEditors/SequencerKeyEditor.h
--- a/Engine/Source/Editor/MovieSceneTools/Private/CurveKeyEditors/SequencerKeyEditor.h
+++ b/Engine/Source/Editor/MovieSceneTools/Private/CurveKeyEditors/SequencerKeyEditor.h
@@ -261,6 +261,14 @@ struct TSequencerKeyEditor
 		Owner->Modify();
 		Owner->SetFlags(RF_Transactional);
 
+		// When the owning object is different from the section (e.g., decoration
+		// channels), also modify the section so its signature updates and
+		// FChannelModel's cached keys detect the change.
+		if (Owner != OwningSection)
+		{
+			OwningSection->Modify();
+		}
+
 		const bool bAutoSetTrackDefaults = Sequencer->GetAutoSetTrackDefaults();
 		const bool bPreserveRatio = ChannelMetaData ? ChannelMetaData->bPreserveRatio : false;
 		const bool bCanPreserveRatio = ChannelMetaData ? ChannelMetaData->bCanPreserveRatio : false;
