Anim In Engine:  System for Auto Baking animation sequences for skel mesh components that allow for enabling quick worfklows for going back and forth between live and cached data for cinematic and physics workflows. Includes main ISequencerBaker/FSequencerBaker singleton that contains a set of recorders(ISequencerBakeRecorder) that can record data over a range of time, while performing Sequencer evaluations. The UAnimSequenceBaker implements an animation sequence recorder. This system then integrates with IMovieSceneLinkedAnimTrackProvider, which provides the functionality to create linked anim sequences and corresponding tracks and haven them getting disabled/enabled while hooked up to the CommonAnimationTrackEditor.  Then depending upon if we are integrating with old animation track, DefaultLinkedAnimTrackProvider, or the new anim mixer, MixerLinkedAnimTrackProvider, we perform different set up actions to set up the links and the tracks, and do the isolation toggling.

#Synced-CL 52367967

diff --git a/Engine/Plugins/Animation/ControlRig/Source/ControlRigEditor/Private/Sequencer/ControlRigParameterTrackEditor.cpp b/Engine/Plugins/Animation/ControlRig/Source/ControlRigEditor/Private/Sequencer/ControlRigParameterTrackEditor.cpp
--- a/Engine/Plugins/Animation/ControlRig/Source/ControlRigEditor/Private/Sequencer/ControlRigParameterTrackEditor.cpp
+++ b/Engine/Plugins/Animation/ControlRig/Source/ControlRigEditor/Private/Sequencer/ControlRigParameterTrackEditor.cpp
@@ -2,6 +2,8 @@
 
 #include "Sequencer/ControlRigParameterTrackEditor.h"
 #include "Sequencer/MovieSceneControlRigSystem.h"
+#include "Baker/SequencerBakerSubsystem.h"
+#include "Baker/SequencerBaker.h"
 #include "MVVM/Extensions/ITrackExtension.h"
 #include "MVVM/Extensions/ISectionOwnerExtension.h"
 #include "MVVM/SectionModelStorageExtension.h"
@@ -423,7 +425,7 @@ void FControlRigParameterTrackEditor::OnInitialize()
 					bHasEnteredSilent = true;
 				}
 
-				IterateTracks([this, &OldToNewControlRigs, ControlRigEditMode, &bRequestEvaluate](UMovieSceneControlRigParameterTrack* Track)
+				IterateTracks([this, &OldToNewControlRigs, ControlRigEditMode, &bRequestEvaluate](UMovieSceneControlRigParameterTrack* Track, FGuid InGuid)
 				{
 					if (UControlRig* OldControlRig = Track->GetControlRig())
 					{
@@ -440,7 +442,7 @@ void FControlRigParameterTrackEditor::OnInitialize()
 							{
 								Track->Modify();
 								Track->ReplaceControlRig(NewControlRig, OldControlRig->GetAssetReference() != NewControlRig->GetAssetReference());
-								BindControlRig(NewControlRig);
+								BindControlRig(NewControlRig,InGuid);
 								bRequestEvaluate = true;
 							}
 							else
@@ -534,7 +536,7 @@ void FControlRigParameterTrackEditor::OnInitialize()
 	}
 
 	// Register all modified/selections for control rigs
-	IterateTracks([this](UMovieSceneControlRigParameterTrack* Track)
+	IterateTracks([this](UMovieSceneControlRigParameterTrack* Track, FGuid InGuid)
 	{
 			
 		if (UControlRig* ControlRig = Track->GetControlRig())
@@ -546,7 +548,7 @@ void FControlRigParameterTrackEditor::OnInitialize()
 			// that the track can pull information from latest version of the rig after construction event
 			Track->BindControlRigDelegates();
 			
-			BindControlRig(ControlRig);
+			BindControlRig(ControlRig, InGuid);
 		}
 
 		// Mark layered mode on track color and display name 
@@ -574,7 +576,7 @@ FControlRigParameterTrackEditor::~FControlRigParameterTrackEditor()
 }
 
 
-void FControlRigParameterTrackEditor::BindControlRig(UControlRig* ControlRig)
+void FControlRigParameterTrackEditor::BindControlRig(UControlRig* ControlRig, FGuid InGuid)
 {
 	if (ControlRig)
 	{
@@ -586,7 +588,7 @@ void FControlRigParameterTrackEditor::BindControlRig(UControlRig* ControlRig)
 			ControlRig->ControlUndoBracket().AddRaw(this, &FControlRigParameterTrackEditor::HandleControlUndoBracket);
 			ControlRig->ControlRigBound().AddRaw(this, &FControlRigParameterTrackEditor::HandleOnControlRigBound);
 		
-			BoundControlRigs.Add(ControlRig);
+			BoundControlRigs.Add(ControlRig, InGuid);
 		}
 
 		if (UMovieSceneControlRigParameterTrack* Track = FindTrack(ControlRig))
@@ -651,16 +653,16 @@ void FControlRigParameterTrackEditor::UnbindControlRig(UControlRig* ControlRig)
 void FControlRigParameterTrackEditor::UnbindAllControlRigs()
 {
 	ClearOutAllSpaceAndConstraintDelegates();
-	TArray<TWeakObjectPtr<UControlRig>> ControlRigs = BoundControlRigs;
-	for (TWeakObjectPtr<UControlRig>& ObjectPtr : ControlRigs)
+	TMap<TWeakObjectPtr<UControlRig>,FGuid> ControlRigs = BoundControlRigs;
+	for (TPair<TWeakObjectPtr<UControlRig>, FGuid>& ObjectPtr : ControlRigs)
 	{
-		if (ObjectPtr.IsValid())
+		if (ObjectPtr.Key.IsValid())
 		{
-			UControlRig* ControlRig = ObjectPtr.Get();
+			UControlRig* ControlRig = ObjectPtr.Key.Get();
 			UnbindControlRig(ControlRig);
 		}
 	}
-	BoundControlRigs.SetNum(0);
+	BoundControlRigs.Reset();
 }
 
 
@@ -669,7 +671,15 @@ void FControlRigParameterTrackEditor::ObjectImplicitlyAdded(UObject* InObject)
 	UControlRig* ControlRig = Cast<UControlRig>(InObject);
 	if (ControlRig)
 	{
-		BindControlRig(ControlRig);
+		IterateTracks([this, ControlRig](UMovieSceneControlRigParameterTrack* Track, FGuid InGuid)
+			{
+				if (ControlRig == Track->GetControlRig())
+				{
+					BindControlRig(ControlRig, InGuid);
+				}
+
+				return false;
+			});
 	}
 }
 
@@ -702,19 +712,18 @@ void FControlRigParameterTrackEditor::OnRelease()
 
 	FControlRigEditMode* ControlRigEditMode = GetEditMode();
 	PreviousSelectedControlRigs.Reset();
-
 	TSharedPtr<ISequencer> Sequencer = GetSequencer();
 	
 	if (ControlRigEditMode)
 	{
 		bControlRigEditModeWasOpen = true;
-		for (TWeakObjectPtr<UControlRig>& ControlRig : BoundControlRigs)
+		for (TPair<TWeakObjectPtr<UControlRig>,FGuid>& ControlRig : BoundControlRigs)
 		{
-			if (ControlRig.IsValid())
+			if (ControlRig.Key.IsValid())
 			{
 				TPair<FControlRigAssetStrongReference, TArray<FName>> ClassAndName;
-				ClassAndName.Key = ControlRig->GetAssetReference();
-				ClassAndName.Value = ControlRig->CurrentControlSelection();
+				ClassAndName.Key = ControlRig.Key->GetAssetReference();
+				ClassAndName.Value = ControlRig.Key->CurrentControlSelection();
 				PreviousSelectedControlRigs.Add(ClassAndName);
 			}
 		}
@@ -1413,7 +1422,7 @@ void FControlRigParameterTrackEditor::BakeToControlRig(const FControlRigAssetStr
 					{
 						ControlRigEditMode->AddControlRigObject(ControlRig, Sequencer);
 					}
-					BindControlRig(ControlRig);
+					BindControlRig(ControlRig, ObjectBinding);
 
 					TempAnimSequence->MarkAsGarbage();
 					AnimSeqExportOption->MarkAsGarbage();
@@ -1424,7 +1433,7 @@ void FControlRigParameterTrackEditor::BakeToControlRig(const FControlRigAssetStr
 	}
 }
 
-static void IterateTracksInMovieScene(const UMovieScene& InMovieScene, TFunctionRef<bool(UMovieSceneControlRigParameterTrack*)> InCallback)
+static void IterateTracksInMovieScene(const UMovieScene& InMovieScene, TFunctionRef<bool(UMovieSceneControlRigParameterTrack*, FGuid)> InCallback)
 {
 	const TArray<FMovieSceneBinding>& Bindings = InMovieScene.GetBindings();
 	for (const FMovieSceneBinding& Binding : Bindings)
@@ -1434,7 +1443,7 @@ static void IterateTracksInMovieScene(const UMovieScene& InMovieScene, TFunction
 		{
 			if (UMovieSceneControlRigParameterTrack* CRTrack = Cast<UMovieSceneControlRigParameterTrack>(Track))
 			{
-				InCallback(CRTrack);
+				InCallback(CRTrack, Binding.GetObjectGuid());
 			}
 		}
 
@@ -1450,7 +1459,7 @@ static void IterateTracksInMovieScene(const UMovieScene& InMovieScene, TFunction
 				{
 					if (UMovieSceneControlRigParameterTrack* CRTrack = Cast<UMovieSceneControlRigParameterTrack>(ChildTrack))
 					{
-						InCallback(CRTrack);
+						InCallback(CRTrack, Binding.GetObjectGuid());
 					}
 				}
 			}
@@ -1461,12 +1470,13 @@ static void IterateTracksInMovieScene(const UMovieScene& InMovieScene, TFunction
 	{
 		if (UMovieSceneControlRigParameterTrack* CRTrack = Cast<UMovieSceneControlRigParameterTrack>(Track))
 		{
-			InCallback(CRTrack);
+			FGuid Empty;
+			InCallback(CRTrack,Empty);
 		}
 	}
 }
 
-void FControlRigParameterTrackEditor::IterateTracks(TFunctionRef<bool(UMovieSceneControlRigParameterTrack*)> Callback) const
+void FControlRigParameterTrackEditor::IterateTracks(TFunctionRef<bool(UMovieSceneControlRigParameterTrack*, FGuid)> Callback) const
 {
 	const TSharedPtr<ISequencer> Sequencer = GetSequencer();
 	UMovieSceneSequence* OwnerSequence = Sequencer ? Sequencer->GetFocusedMovieSceneSequence() : nullptr;
@@ -1946,7 +1956,7 @@ void FControlRigParameterTrackEditor::HandleAddTrackSubMenu(FMenuBuilder& MenuBu
 	{
 		TArray<FControlRigAssetStrongReference> ExistingRigs;
 		USkeletalMeshComponent* SkeletalMeshComponent = AcquireSkeletalMeshFromObject(BoundObject, ParentSequencer);
-		IterateTracks([&ExistingRigs, SkeletalMeshComponent](UMovieSceneControlRigParameterTrack* Track) -> bool
+		IterateTracks([&ExistingRigs, SkeletalMeshComponent](UMovieSceneControlRigParameterTrack* Track, FGuid InGuid) -> bool
 		{
 			if (UControlRig* ControlRig = Track->GetControlRig())
 			{
@@ -2220,7 +2230,7 @@ bool FControlRigParameterTrackEditor::IsControlRigAllowed(const FAssetData& Asse
 static TArray<UMovieSceneControlRigParameterTrack*> GetExistingControlRigTracksForSkeletalMeshComponent(UMovieScene* MovieScene, USkeletalMeshComponent* SkeletalMeshComponent)
 {
 	TArray<UMovieSceneControlRigParameterTrack*> ExistingControlRigTracks;
-	IterateTracksInMovieScene(*MovieScene, [&ExistingControlRigTracks, SkeletalMeshComponent](UMovieSceneControlRigParameterTrack* Track) -> bool
+	IterateTracksInMovieScene(*MovieScene, [&ExistingControlRigTracks, SkeletalMeshComponent](UMovieSceneControlRigParameterTrack* Track, FGuid InGuid) -> bool
 	{
 		if (UControlRig* ControlRig = Track->GetControlRig())
 		{
@@ -2436,7 +2446,7 @@ void FControlRigParameterTrackEditor::AddControlRig(const FControlRigAssetStrong
 	//5.4.4 we can't change headers so for now we just make the change here locally
 	if (UMovieSceneTrack* Track = FindOrCreateControlRigTrack(Sequencer, World, InSource, BindingProxy, bIsLayeredControlRig, InTrackCreationCallback))
 	{
-		BindControlRig(CastChecked<UMovieSceneControlRigParameterTrack>(Track)->GetControlRig());
+		BindControlRig(CastChecked<UMovieSceneControlRigParameterTrack>(Track)->GetControlRig(), ObjectBinding);
 	}
 }
 
@@ -2959,11 +2969,11 @@ void FControlRigParameterTrackEditor::HandleActorAdded(AActor* Actor, FGuid Targ
 
 void FControlRigParameterTrackEditor::OnActivateSequenceChanged(FMovieSceneSequenceIDRef ID)
 {
-	IterateTracks([this](UMovieSceneControlRigParameterTrack* Track)
+	IterateTracks([this](UMovieSceneControlRigParameterTrack* Track, FGuid InGuid)
 	{
 		if (UControlRig* ControlRig = Track->GetControlRig())
 		{
-			BindControlRig(ControlRig);
+			BindControlRig(ControlRig, InGuid);
 		}
 		return false;
 	});
@@ -3023,9 +3033,9 @@ void FControlRigParameterTrackEditor::OnActivateSequenceChanged(FMovieSceneSeque
 						bool bSequencerSet = false;
 
 						TSharedPtr<ISequencer> Sequencer = SharedThis->GetSequencer();
-						for (TWeakObjectPtr<UControlRig>& ControlRig : SharedThis->BoundControlRigs)
+						for (TPair<TWeakObjectPtr<UControlRig>,FGuid>& ControlRig : SharedThis->BoundControlRigs)
 						{
-							if (UControlRig* Rig = ControlRig.Get())
+							if (UControlRig* Rig = ControlRig.Key.Get())
 							{
 								ControlRigEditMode->AddControlRigObject(Rig, Sequencer);
 								bSequencerSet = true;
@@ -3109,7 +3119,7 @@ void FControlRigParameterTrackEditor::OnSequencerDataChanged(EMovieSceneDataChan
 				ControlRig->SetFramesPerSecond(FPS);
 
 				bool bControlRigInTrack = false;
-				IterateTracks([&bControlRigInTrack, ControlRig](const UMovieSceneControlRigParameterTrack* Track)
+				IterateTracks([&bControlRigInTrack, ControlRig](const UMovieSceneControlRigParameterTrack* Track, FGuid InGuid)
 				{
 					if(Track->GetControlRig() == ControlRig)
 			        {
@@ -3134,7 +3144,7 @@ void FControlRigParameterTrackEditor::PostEvaluation(UMovieScene* MovieScene, FF
 {
 	if (MovieScene)
 	{
-		IterateTracksInMovieScene(*MovieScene, [this](const UMovieSceneControlRigParameterTrack* Track)
+		IterateTracksInMovieScene(*MovieScene, [this](const UMovieSceneControlRigParameterTrack* Track, FGuid InGuid)
 		{
 			if (const UControlRig* ControlRig = Track->GetControlRig())
 				{
@@ -3656,7 +3666,7 @@ FMovieSceneTrackEditor::FFindOrCreateTrackResult FControlRigParameterTrackEditor
 	FFindOrCreateTrackResult Result;
 	bool bTrackExisted = false;
 
-	IterateTracks([&Result, &bTrackExisted, ControlRig](UMovieSceneControlRigParameterTrack* Track)
+	IterateTracks([&Result, &bTrackExisted, ControlRig](UMovieSceneControlRigParameterTrack* Track, FGuid InGuid)
 	{
 		if (Track->GetControlRig() == ControlRig)
 		{
@@ -3689,7 +3699,7 @@ UMovieSceneControlRigParameterTrack* FControlRigParameterTrackEditor::FindTrack(
 	}
 	
 	UMovieSceneControlRigParameterTrack* FoundTrack = nullptr;
-	IterateTracks([InControlRig, &FoundTrack](UMovieSceneControlRigParameterTrack* Track)
+	IterateTracks([InControlRig, &FoundTrack](UMovieSceneControlRigParameterTrack* Track, FGuid InGuid)
 	{
 		if (Track->GetControlRig() == InControlRig)
 		{
@@ -3823,7 +3833,7 @@ void FControlRigParameterTrackEditor::ClearOutAllSpaceAndConstraintDelegates(con
 		return;
 	}
 
-	IterateTracks([InOptionalControlRig](const UMovieSceneControlRigParameterTrack* Track)
+	IterateTracks([InOptionalControlRig](const UMovieSceneControlRigParameterTrack* Track, FGuid InGuid)
 	{
 		if (InOptionalControlRig && Track->GetControlRig() != InOptionalControlRig)
 		{
@@ -4348,12 +4358,33 @@ void FControlRigParameterTrackEditor::HandleControlSelected(UControlRig* Subject
 		}
 	}
 	
+	const TSharedPtr<ISequencer> Sequencer = GetSequencer();
+	if (!Sequencer)
+	{
+		return;
+	}
+	if (bSelected) //if selected
+	{
+		if (FGuid* BindingID = BoundControlRigs.Find(Subject))
+		{
+			if (USequencerBakeSubsystem* SequencerBakeSubsystem = GEditor->GetEditorSubsystem<USequencerBakeSubsystem>())
+			{
+				if (TSharedPtr<UE::Sequencer::FSequencerBaker> SequencerBaker = SequencerBakeSubsystem->GetSequencerBaker(Sequencer))
+				{
+					constexpr bool bIsolate = false;
+					SequencerBaker->IsolateRecorderWithGuid(*BindingID, bIsolate);
+				}
+			}
+		}
+	}
+
 	if (bIsDoingSelection)
 	{
 		return;
 	}
 	
 	TGuardValue<bool> Guard(bIsDoingSelection, true);
+	
 
 	const FName ControlRigName(*Subject->GetName());
 	if (TSharedPtr<IControlRigObjectBinding> ObjectBinding = Subject->GetObjectBinding())
@@ -4369,8 +4400,7 @@ void FControlRigParameterTrackEditor::HandleControlSelected(UControlRig* Subject
 		FFindOrCreateTrackResult TrackResult = FindOrCreateControlRigTrackForObject(HandleResult.Handle, Subject, ControlRigName, bCreateTrack);
 		UMovieSceneControlRigParameterTrack* Track = CastChecked<UMovieSceneControlRigParameterTrack>(TrackResult.Track, ECastCheckedType::NullAllowed);
 
-		const TSharedPtr<ISequencer> Sequencer = GetSequencer();
-		if (Track && Sequencer)
+		if (Track)
 		{
 			//Just select in section to key, if deselecting makes sure deselected everywhere
 			if (bSelected == true)
@@ -4391,6 +4421,7 @@ void FControlRigParameterTrackEditor::HandleControlSelected(UControlRig* Subject
 			}
 
 			SetUpEditModeIfNeeded(Subject);
+			
 
 			//Force refresh later, not now
 			bSkipNextSelectionFromTimer = bSkipNextSelectionFromTimer ||
@@ -4445,7 +4476,7 @@ void FControlRigParameterTrackEditor::HandleControlModified(UControlRig* Control
 
 	FTransform  Transform = ControlRig->GetControlLocalTransform(ControlElement->GetFName());
 
-	IterateTracks([this, ControlRig, ControlElement, Context](UMovieSceneControlRigParameterTrack* Track)
+	IterateTracks([this, ControlRig, ControlElement, Context](UMovieSceneControlRigParameterTrack* Track, FGuid InGuid)
 	{
 		if (Track && Track->GetControlRig() == ControlRig)
 		{
@@ -4538,23 +4569,23 @@ void FControlRigParameterTrackEditor::HandleOnObjectBoundToControlRig(UObject* I
 	TMap<TWeakObjectPtr<UControlRig>, TArray<FName>> ReselectIfNeeded;
 	// look for sections to update
 	TArray<UMovieSceneControlRigParameterSection*> SectionsToUpdate;
-	for (TWeakObjectPtr<UControlRig>& ControlRigPtr : BoundControlRigs)
+	for (TPair<TWeakObjectPtr<UControlRig>, FGuid>& ControlRigPair : BoundControlRigs)
 	{
-		if (ControlRigPtr.IsValid() == false)
+		TStrongObjectPtr<UControlRig> ControlRig = ControlRigPair.Key.Pin();
+		if (ControlRig.IsValid() == false)
 		{
 			continue;
 		}
-		TArray<FName> Selection = ControlRigPtr->CurrentControlSelection();
+		TArray<FName> Selection = ControlRig->CurrentControlSelection();
 		if (Selection.Num() > 0)
 		{
-			ReselectIfNeeded.Add(ControlRigPtr, Selection);
+			ReselectIfNeeded.Add(ControlRigPair.Key, Selection);
 		}
-		const TSharedPtr<IControlRigObjectBinding> Binding =
-			ControlRigPtr.IsValid() ? ControlRigPtr->GetObjectBinding() : nullptr;
+		const TSharedPtr<IControlRigObjectBinding> Binding = ControlRig->GetObjectBinding();
 		const UObject* CurrentObject = Binding ? Binding->GetBoundObject() : nullptr;
 		if (CurrentObject == InObject)
 		{
-			if (const UMovieSceneControlRigParameterTrack* Track = FindTrack(ControlRigPtr.Get()))
+			if (const UMovieSceneControlRigParameterTrack* Track = FindTrack(ControlRig.Get()))
 			{
 				for (UMovieSceneSection* BaseSection : Track->GetAllSections())
 				{
@@ -5290,7 +5321,7 @@ void FControlRigParameterTrackEditor::ProcessKeyOperation(FFrameNumber InKeyTime
 		// Cache the evaluation tracks if requested
 		if (UMovieSceneControlRigParameterTrack::ShouldCacheLegacyTemplate())
 		{
-			IterateTracks([&CachedEvaluationTracks](UMovieSceneControlRigParameterTrack* Track)
+			IterateTracks([&CachedEvaluationTracks](UMovieSceneControlRigParameterTrack* Track, FGuid InGuid)
 		   {
 			   if (Track && !CachedEvaluationTracks.Contains(Track))
 			   {
diff --git a/Engine/Plugins/Animation/ControlRig/Source/ControlRigEditor/Private/Sequencer/ControlRigParameterTrackEditor.h b/Engine/Plugins/Animation/ControlRig/Source/ControlRigEditor/Private/Sequencer/ControlRigParameterTrackEditor.h
--- a/Engine/Plugins/Animation/ControlRig/Source/ControlRigEditor/Private/Sequencer/ControlRigParameterTrackEditor.h
+++ b/Engine/Plugins/Animation/ControlRig/Source/ControlRigEditor/Private/Sequencer/ControlRigParameterTrackEditor.h
@@ -336,7 +336,7 @@ public:
 	void SetUpEditModeIfNeeded(UControlRig* ControlRig);
 
 	/** Helper functions to iterate over UMovieSceneControlRigParameterTracks in the currently focussed MovieScene*/
-	void IterateTracks(TFunctionRef<bool(UMovieSceneControlRigParameterTrack*)> Callback) const;
+	void IterateTracks(TFunctionRef<bool(UMovieSceneControlRigParameterTrack*, FGuid)> Callback) const;
 private:
 	/** Command Bindings added by the Transform Track Editor to Sequencer and curve editor. */
 	TSharedPtr<FUICommandList> CommandBindings;
@@ -405,10 +405,10 @@ private:
 	TSet<FName> DisplayedControls;
 	TSet<FName> UnDisplayedControls;
 	bool bCurveDisplayTickIsPending;
-	void BindControlRig(UControlRig* ControlRig);
+	void BindControlRig(UControlRig* ControlRig, FGuid);
 	void UnbindControlRig(UControlRig* ControlRig);
 	void UnbindAllControlRigs();
-	TArray<TWeakObjectPtr<UControlRig>> BoundControlRigs;
+	TMap<TWeakObjectPtr<UControlRig>, FGuid> BoundControlRigs;
 
 private:
 
diff --git a/Engine/Plugins/Experimental/MovieSceneAnimMixer/Source/MovieSceneAnimMixer/Private/MovieSceneAnimationMixerLayer.cpp b/Engine/Plugins/Experimental/MovieSceneAnimMixer/Source/MovieSceneAnimMixer/Private/MovieSceneAnimationMixerLayer.cpp
--- a/Engine/Plugins/Experimental/MovieSceneAnimMixer/Source/MovieSceneAnimMixer/Private/MovieSceneAnimationMixerLayer.cpp
+++ b/Engine/Plugins/Experimental/MovieSceneAnimMixer/Source/MovieSceneAnimMixer/Private/MovieSceneAnimationMixerLayer.cpp
@@ -98,7 +98,7 @@ FText UMovieSceneAnimationMixerLayer::GetDisplayName() const
 	// Generate default name based on content
 	if (ChildTrack)
 	{
-		return FText::Format(NSLOCTEXT("MovieSceneAnimationMixerLayer", "ChildTrackLayerNameFormat", "{0}: {1}"), LayerNumberText, FText::FromString(ChildTrack->GetTrackName().ToString()));		
+		return FText::Format(NSLOCTEXT("MovieSceneAnimationMixerLayer", "ChildTrackLayerNameFormat", "{0}: {1}"), LayerNumberText, ChildTrack->GetDisplayName());
 	}
 	else
 	{
diff --git a/Engine/Plugins/Experimental/MovieSceneAnimMixer/Source/MovieSceneAnimMixerEditor/MovieSceneAnimMixerEditor.Build.cs b/Engine/Plugins/Experimental/MovieSceneAnimMixer/Source/MovieSceneAnimMixerEditor/MovieSceneAnimMixerEditor.Build.cs
--- a/Engine/Plugins/Experimental/MovieSceneAnimMixer/Source/MovieSceneAnimMixerEditor/MovieSceneAnimMixerEditor.Build.cs
+++ b/Engine/Plugins/Experimental/MovieSceneAnimMixer/Source/MovieSceneAnimMixerEditor/MovieSceneAnimMixerEditor.Build.cs
@@ -1,4 +1,4 @@
-﻿// Copyright Epic Games, Inc. All Rights Reserved.
+// Copyright Epic Games, Inc. All Rights Reserved.
 
 namespace UnrealBuildTool.Rules
 {
@@ -29,6 +29,7 @@ namespace UnrealBuildTool.Rules
 					"Engine",
 					"InputCore",
 					"LevelEditor",
+					"LevelSequence",
 					"Persona",
 					"MovieScene",
 					"MovieSceneAnimMixer",
diff --git a/Engine/Plugins/Experimental/MovieSceneAnimMixer/Source/MovieSceneAnimMixerEditor/Private/MixerLinkedAnimTrackProvider.cpp b/Engine/Plugins/Experimental/MovieSceneAnimMixer/Source/MovieSceneAnimMixerEditor/Private/MixerLinkedAnimTrackProvider.cpp
new file mode 100644
--- /dev/null
+++ b/Engine/Plugins/Experimental/MovieSceneAnimMixer/Source/MovieSceneAnimMixerEditor/Private/MixerLinkedAnimTrackProvider.cpp
@@ -0,0 +1,421 @@
+// Copyright Epic Games, Inc. All Rights Reserved.
+
+#include "MixerLinkedAnimTrackProvider.h"
+#include "MovieSceneAnimationMixerTrackEditor.h"
+#include "MovieSceneAnimationMixerTrack.h"
+#include "MovieSceneAnimationMixerLayer.h"
+#include "TrackEditors/CommonAnimationTrackEditor.h"
+#include "ISequencer.h"
+#include "MovieScene.h"
+#include "LevelSequence.h"
+#include "LevelSequenceAnimSequenceLink.h"
+#include "Tracks/MovieSceneSkeletalAnimationTrack.h"
+#include "Sections/MovieSceneSkeletalAnimationSection.h"
+#include "Animation/AnimSequence.h"
+#include "Components/SkeletalMeshComponent.h"
+#include "Baker/SequencerBakerSubsystem.h"
+#include "Baker/SequencerBaker.h"
+#include "Editor.h"
+
+#define LOCTEXT_NAMESPACE "MixerLinkedAnimTrackProvider"
+
+using namespace UE::Sequencer;
+
+static bool DisableMixerTrackHelper(UMovieSceneSignedObject* InObject,  bool bInEvalDisabled)
+{
+	bool bWasChanged = false;
+	if (UMovieSceneTrack* Track = Cast<UMovieSceneTrack>(InObject))
+	{
+		if (Track->GetMaxRowIndex() == 0)
+		{
+			if (Track->IsEvalDisabled() != bInEvalDisabled)
+			{
+				bWasChanged = true;
+				Track->Modify();
+				Track->SetEvalDisabled(bInEvalDisabled);
+			}
+		}
+		else
+		{
+			for (int32 RowIndex = 0; RowIndex <= Track->GetMaxRowIndex(); ++RowIndex)
+			{
+				if (Track->IsRowEvalDisabled(RowIndex) != bInEvalDisabled)
+				{
+					bWasChanged = true;
+					Track->Modify();
+					Track->SetRowEvalDisabled(bInEvalDisabled, RowIndex);
+				}
+			}
+		}
+	}
+	else if (UMovieSceneSection* Section = Cast<UMovieSceneSection>(InObject))
+	{
+		if (UMovieSceneTrack* SectionParent = Section->GetTypedOuter<UMovieSceneTrack>())
+		{
+			int32 RowIndex = Section->GetRowIndex();
+			if (SectionParent->IsRowEvalDisabled(RowIndex) != bInEvalDisabled)
+			{
+				bWasChanged = true;
+				SectionParent->Modify();
+				SectionParent->SetRowEvalDisabled(bInEvalDisabled, RowIndex);
+			}
+		}
+	}
+	return bWasChanged;
+}
+
+bool FMixerLinkedAnimTrackProvider::HasMixerTrack(const TWeakPtr<ISequencer>& InSequencer, FGuid InBindingID) const
+{
+	TSharedPtr<ISequencer> SequencerPtr = InSequencer.Pin();
+	if (!SequencerPtr.IsValid() || !SequencerPtr->GetFocusedMovieSceneSequence() || !SequencerPtr->GetFocusedMovieSceneSequence()->GetMovieScene())
+	{
+		return false;
+	}
+	UMovieScene* MovieScene = SequencerPtr->GetFocusedMovieSceneSequence()->GetMovieScene();
+	TArray<UMovieSceneTrack*> MixerTracks = MovieScene->FindTracks(UMovieSceneAnimationMixerTrack::StaticClass(), InBindingID, NAME_None);
+	for (UMovieSceneTrack* Track : MixerTracks)
+	{
+		if (UMovieSceneAnimationMixerTrack* MixerTrack = Cast<UMovieSceneAnimationMixerTrack>(Track))
+		{
+			return true;
+		}
+	}
+	return false;
+}
+
+TPair<UMovieSceneAnimationMixerTrack*, UMovieSceneSection*> FMixerLinkedAnimTrackProvider::FindMixerTrack(const TWeakPtr<ISequencer>& InSequencer, FGuid InBindingID) const
+{
+	TPair<UMovieSceneAnimationMixerTrack*, UMovieSceneSection*> Pair(nullptr, nullptr);
+		
+	TSharedPtr<ISequencer> SequencerPtr = InSequencer.Pin();
+	if (!SequencerPtr.IsValid() || !SequencerPtr->GetFocusedMovieSceneSequence() || !SequencerPtr->GetFocusedMovieSceneSequence()->GetMovieScene())
+	{
+		return Pair;
+	}
+
+	TArray<FLevelSequenceAnimSequenceLinkItem*> LinkedItems = FCommonAnimationTrackEditor::GetLinkedAnimSequences(InSequencer, InBindingID);
+	UMovieScene* MovieScene = SequencerPtr->GetFocusedMovieSceneSequence()->GetMovieScene();
+	TArray<UMovieSceneTrack*> MixerTracks = MovieScene->FindTracks(UMovieSceneAnimationMixerTrack::StaticClass(), InBindingID, NAME_None);
+
+	for (UMovieSceneTrack* Track : MixerTracks)
+	{
+		if (UMovieSceneAnimationMixerTrack* MixerTrack = Cast<UMovieSceneAnimationMixerTrack>(Track))
+		{
+			if (UMovieSceneSkeletalAnimationTrack* ChildTrack = FindLinkedChildTrack(MixerTrack, LinkedItems))
+			{
+				TArray<UMovieSceneSection*> Sections = ChildTrack->GetAllSections();
+				for (int32 Index = 0; Index < LinkedItems.Num(); ++Index)
+				{
+					FLevelSequenceAnimSequenceLinkItem* Item = LinkedItems[Index];
+					if (Item && Item->bAutoBake == true)
+					{
+						UAnimSequence* LinkedAnimSequence = Item->ResolveAnimSequence();
+						for (int32 SectionIndex = 0; SectionIndex < Sections.Num(); ++SectionIndex)
+						{
+							if (UMovieSceneSkeletalAnimationSection* AnimSection = Cast<UMovieSceneSkeletalAnimationSection>(Sections[SectionIndex]))
+							{
+								if (LinkedAnimSequence && LinkedAnimSequence == AnimSection->Params.Animation)
+								{
+									Pair.Key = MixerTrack;
+									Pair.Value = AnimSection;
+									return Pair;
+								}
+							}
+						}
+					}
+				}
+			}
+			if (Pair.Key == nullptr) //pick first one
+			{
+				Pair.Key = MixerTrack;
+			}
+		}
+	}
+	return Pair;
+}
+
+UMovieSceneSkeletalAnimationTrack* FMixerLinkedAnimTrackProvider::FindLinkedChildTrack(
+	UMovieSceneAnimationMixerTrack* MixerTrack,
+	const TArray<FLevelSequenceAnimSequenceLinkItem*>& LinkedItems) const
+{
+	TArray<TObjectPtr<UMovieSceneTrack>> ChildTracks;
+	MixerTrack->GetAllChildTracks(ChildTracks);
+
+	for (const TObjectPtr<UMovieSceneTrack>& ChildTrack : ChildTracks)
+	{
+		if (UMovieSceneSkeletalAnimationTrack* AnimTrack = Cast<UMovieSceneSkeletalAnimationTrack>(ChildTrack.Get()))
+		{
+			const TArray<UMovieSceneSection*>& Sections = AnimTrack->GetAllSections();
+			for (FLevelSequenceAnimSequenceLinkItem* Item : LinkedItems)
+			{
+				if (Item && Item->bAutoBake)
+				{
+					UAnimSequence* LinkedAnimSequence = Item->ResolveAnimSequence();
+					for (UMovieSceneSection* Section : Sections)
+					{
+						if (UMovieSceneSkeletalAnimationSection* AnimSection = Cast<UMovieSceneSkeletalAnimationSection>(Section))
+						{
+							if (LinkedAnimSequence && LinkedAnimSequence == AnimSection->Params.Animation)
+							{
+								return AnimTrack;
+							}
+						}
+					}
+				}
+			}
+		}
+	}
+	return nullptr;
+}
+
+bool FMixerLinkedAnimTrackProvider::CanHandleBinding(const TWeakPtr<ISequencer>& InSequencer, FGuid InBindingID) const
+{
+	return HasMixerTrack(InSequencer, InBindingID);
+}
+
+//this function expects a valid mixer(anim skel mesh) track WITH the linked anim sequence, which will only happen if the section(value) is valid
+UMovieSceneTrack* FMixerLinkedAnimTrackProvider::GetLinkedAnimTrack(const TWeakPtr<ISequencer>& InSequencer, FGuid InBindingID) const
+{
+	TPair<UMovieSceneAnimationMixerTrack*, UMovieSceneSection*> MixerTrack = FindMixerTrack(InSequencer, InBindingID);
+	if (MixerTrack.Value)
+	{
+		if (UMovieSceneTrack* Track = MixerTrack.Value->GetTypedOuter<UMovieSceneTrack>())
+		{
+			return Track;
+		}
+	}
+	return nullptr;
+}
+
+UMovieSceneTrack* FMixerLinkedAnimTrackProvider::GetOrCreateLinkedAnimTrack(
+	const TWeakPtr<ISequencer>& InSequencer,
+	FGuid InBindingID,
+	FCommonAnimationTrackEditor& InTrackEditor)
+{
+
+	TSharedPtr<ISequencer> SequencerPtr = InSequencer.Pin();
+	if (!SequencerPtr.IsValid() || !SequencerPtr->GetFocusedMovieSceneSequence() || !SequencerPtr->GetFocusedMovieSceneSequence()->GetMovieScene())
+	{
+		return nullptr;
+	}
+
+	TPair<UMovieSceneAnimationMixerTrack*, UMovieSceneSection*> MixerTrackSection = FindMixerTrack(InSequencer, InBindingID);
+	UMovieSceneAnimationMixerTrack* MixerTrack = MixerTrackSection.Key;
+	if (MixerTrack == nullptr)
+	{
+		return nullptr; //shouldn't have we have at least one mixer track
+	}
+	if (MixerTrackSection.Value) //has section, return track
+	{
+		return MixerTrackSection.Key;
+	}
+
+
+	UMovieScene* MovieScene = SequencerPtr->GetFocusedMovieSceneSequence()->GetMovieScene();
+	FScopedTransaction Transaction(LOCTEXT("CreateMixerLinkedAnimTrack_Transaction", "Create Mixer Linked Anim Track"));
+	MovieScene->Modify();
+
+	// Create a child skeletal animation track inside the mixer
+	int32 NextRow = MixerTrack->ComputeNextAvailableRowIndex();
+	UMovieSceneTrack* NewChildTrack = MixerTrack->AddChildTrack(InBindingID, UMovieSceneSkeletalAnimationTrack::StaticClass(), NextRow);
+	UMovieSceneSkeletalAnimationTrack* AnimTrack = Cast<UMovieSceneSkeletalAnimationTrack>(NewChildTrack);
+	if (!AnimTrack)
+	{
+		return nullptr;
+	}
+
+	AnimTrack->Modify();
+	const FString TrackName = TEXT("(Linked)");
+	AnimTrack->SetDisplayName(FText::FromString(TrackName));
+	AnimTrack->SetColorTint(UMovieSceneSkeletalAnimationTrack::LinkedAnimSeqTrackColor);
+	USkeletalMeshComponent* SkelMeshComp = FCommonAnimationTrackEditor::AcquireSkeletalMeshFromObjectGuid(InBindingID, SequencerPtr);
+	USkeleton* Skeleton = FCommonAnimationTrackEditor::AcquireSkeletonFromObjectGuid(InBindingID, SequencerPtr);
+
+	// Check if we already have a linked anim item to use, or need to create one
+	bool bNeedToCreateLinkedAnim = true;
+	int32 LinkedItemNum = INDEX_NONE;
+	TArray<FLevelSequenceAnimSequenceLinkItem*> LinkedItems = FCommonAnimationTrackEditor::GetLinkedAnimSequences(SequencerPtr, InBindingID);
+	if (LinkedItems.Num() > 0)
+	{
+		for (int32 LinkedIndex = 0; LinkedIndex < LinkedItems.Num(); ++LinkedIndex)
+		{
+			FLevelSequenceAnimSequenceLinkItem* LinkedItem = LinkedItems[LinkedIndex];
+			if (LinkedItem && LinkedItem->bUseCustomTimeRange == false)
+			{
+				bNeedToCreateLinkedAnim = false;
+				LinkedItemNum = LinkedIndex;
+			}
+		}
+	}
+	if (bNeedToCreateLinkedAnim)
+	{
+		InTrackEditor.CreateLinkedAnimationSequence(SkelMeshComp, Skeleton, InBindingID, bNeedToCreateLinkedAnim);
+		LinkedItems = FCommonAnimationTrackEditor::GetLinkedAnimSequences(SequencerPtr, InBindingID);
+		LinkedItemNum = LinkedItems.Num() - 1;
+	}
+
+	if (LinkedItems.Num() > 0 && LinkedItemNum != INDEX_NONE)
+	{
+		if (FLevelSequenceAnimSequenceLinkItem* LinkedItem = LinkedItems[LinkedItemNum])
+		{
+			UAnimSequence* AnimSequence = LinkedItem->ResolveAnimSequence();
+			TOptional<TRange<FFrameNumber>> OptionalRange = SequencerPtr->GetSubSequenceRange();
+			TRange<FFrameNumber> Range = OptionalRange.IsSet() ? OptionalRange.GetValue() : SequencerPtr->GetFocusedMovieSceneSequence()->GetMovieScene()->GetPlaybackRange();
+			FFrameNumber Frame = Range.GetLowerBoundValue();
+			UMovieSceneSkeletalAnimationSection* NewSection = Cast<UMovieSceneSkeletalAnimationSection>(AnimTrack->AddNewAnimationOnRow(Frame, AnimSequence, 0));
+			if (NewSection)
+			{
+				LinkedItem->bUseCustomTimeRange = false;
+				LinkedItem->bAutoBake = true;
+				NewSection->SetRange(Range);
+				NewChildTrack->SetEvalDisabled(true);
+				return NewChildTrack;
+			}
+		}
+	}
+
+	return nullptr;
+}
+
+bool FMixerLinkedAnimTrackProvider::IsolateLinkedAnimTrack(const TWeakPtr<ISequencer>& InSequencer, FGuid InBindingID, bool bIsolate)
+{
+	bool bWasChanged = false;
+
+	TSharedPtr<ISequencer> SequencerPtr = InSequencer.Pin();
+	if (!SequencerPtr.IsValid() || !SequencerPtr->GetFocusedMovieSceneSequence() || !SequencerPtr->GetFocusedMovieSceneSequence()->GetMovieScene())
+	{
+		return bWasChanged;
+	}
+
+	TArray<FLevelSequenceAnimSequenceLinkItem*> LinkedItems = FCommonAnimationTrackEditor::GetLinkedAnimSequences(InSequencer, InBindingID);
+	if (LinkedItems.Num() == 0)
+	{
+		return bWasChanged;
+	}
+
+	UMovieScene* MovieScene = SequencerPtr->GetFocusedMovieSceneSequence()->GetMovieScene();
+	TArray<UMovieSceneTrack*> AllTracks = MovieScene->FindTracks(UMovieSceneTrack::StaticClass(), InBindingID, NAME_None);
+
+	// Find which MixerTrack contains the linked anim child track, and find the child track
+	UMovieSceneAnimationMixerTrack* LinkedMixerTrack = nullptr;
+	UMovieSceneSkeletalAnimationTrack* ChildMixerTrack = nullptr;
+	for (UMovieSceneTrack* Track : AllTracks)
+	{
+		if (UMovieSceneAnimationMixerTrack* MixerTrack = Cast<UMovieSceneAnimationMixerTrack>(Track))
+		{
+			ChildMixerTrack = FindLinkedChildTrack(MixerTrack, LinkedItems);
+			if(ChildMixerTrack)
+			{
+				LinkedMixerTrack = MixerTrack;
+				break;
+			}
+		}
+	}
+
+	for (UMovieSceneTrack* Track : AllTracks)
+	{
+		if (UMovieSceneAnimationMixerTrack* MixerTrack = Cast<UMovieSceneAnimationMixerTrack>(Track))
+		{
+			//main mixer track, need to disable it's sections when isolation
+			if (MixerTrack == LinkedMixerTrack)
+			{
+				TArray<UMovieSceneSection*> Sections = MixerTrack->GetAllSections();
+				for (UMovieSceneSection* Section : Sections)
+				{
+					if (DisableMixerTrackHelper(Section, bIsolate))
+					{
+						bWasChanged = true;
+					}
+				}
+			}
+			if (MixerTrack != LinkedMixerTrack)
+			{
+				// Other MixerTracks: disable when isolating, enable when un-isolating
+				if (DisableMixerTrackHelper(MixerTrack, bIsolate))
+				{
+					bWasChanged = true;
+				}
+			}
+		}
+		else if (UMovieSceneSkeletalAnimationTrack* AnimTrack = Cast<UMovieSceneSkeletalAnimationTrack>(Track))
+		{
+			// Internal Mixer AnimTrack we enable when isolating
+			if (ChildMixerTrack == AnimTrack)
+			{
+				if (DisableMixerTrackHelper(AnimTrack, !bIsolate))
+				{
+					bWasChanged = true;
+				}
+			}
+			// Standalone anim tracks (not inside a mixer): disable when isolating
+			else if (DisableMixerTrackHelper(AnimTrack, bIsolate))
+			{
+				bWasChanged = true;
+			}
+		}
+		else
+		{
+			// Control rig and other tracks: disable when isolating
+			FString TrackName = Track->GetClass()->GetName();
+			if (TrackName == FString("MovieSceneControlRigParameterTrack"))
+			{
+				if (DisableMixerTrackHelper(Track, bIsolate))
+				{
+					bWasChanged = true;
+				}
+			}
+		}
+	}
+
+	if (USequencerBakeSubsystem* SequencerBakeSubsystem = GEditor->GetEditorSubsystem<USequencerBakeSubsystem>())
+	{
+		if (TSharedPtr<UE::Sequencer::FSequencerBaker> SequencerBaker = SequencerBakeSubsystem->GetSequencerBaker(SequencerPtr))
+		{
+			SequencerBaker->IgnoreLastChange();
+		}
+	}
+
+	return bWasChanged;
+}
+
+void FMixerLinkedAnimTrackProvider::UpdateLinkedAnimSectionRange(const TWeakPtr<ISequencer>& InSequencer, FGuid InBindingID)
+{
+	TSharedPtr<ISequencer> SequencerPtr = InSequencer.Pin();
+	if (!SequencerPtr.IsValid())
+	{
+		return;
+	}
+	ULevelSequence* LevelSequence = Cast<ULevelSequence>(SequencerPtr->GetFocusedMovieSceneSequence());
+	if (!LevelSequence || !LevelSequence->GetClass()->ImplementsInterface(UInterface_AssetUserData::StaticClass()))
+	{
+		return;
+	}
+
+	TOptional<TRange<FFrameNumber>> OptionalRange = SequencerPtr->GetSubSequenceRange();
+	TRange<FFrameNumber> Range = OptionalRange.IsSet() ? OptionalRange.GetValue() : SequencerPtr->GetFocusedMovieSceneSequence()->GetMovieScene()->GetPlaybackRange();
+
+	TArray<FLevelSequenceAnimSequenceLinkItem*> LinkedItems = FCommonAnimationTrackEditor::GetLinkedAnimSequences(InSequencer, InBindingID);
+	if (LinkedItems.Num() == 0)
+	{
+		return;
+	}
+
+	TPair<UMovieSceneAnimationMixerTrack*, UMovieSceneSection*> MixerTrackSection = FindMixerTrack(InSequencer, InBindingID);
+
+	UMovieSceneAnimationMixerTrack* MixerTrack = MixerTrackSection.Key;
+	if (!MixerTrack)
+	{
+		return;
+	}
+	
+	if (UMovieSceneSkeletalAnimationSection* AnimSection = Cast<UMovieSceneSkeletalAnimationSection>(MixerTrackSection.Value))
+	{
+		if (AnimSection->GetRange() != Range)
+		{
+			AnimSection->Modify();
+			AnimSection->SetRange(Range);
+		}
+	}
+}
+
+#undef LOCTEXT_NAMESPACE
diff --git a/Engine/Plugins/Experimental/MovieSceneAnimMixer/Source/MovieSceneAnimMixerEditor/Private/MixerLinkedAnimTrackProvider.h b/Engine/Plugins/Experimental/MovieSceneAnimMixer/Source/MovieSceneAnimMixerEditor/Private/MixerLinkedAnimTrackProvider.h
new file mode 100644
--- /dev/null
+++ b/Engine/Plugins/Experimental/MovieSceneAnimMixer/Source/MovieSceneAnimMixerEditor/Private/MixerLinkedAnimTrackProvider.h
@@ -0,0 +1,46 @@
+// Copyright Epic Games, Inc. All Rights Reserved.
+
+#pragma once
+
+#include "IMovieSceneLinkedAnimTrackProvider.h"
+
+namespace UE::Sequencer { class FAnimationMixerTrackEditor; }
+
+class UMovieSceneAnimationMixerTrack;
+class UMovieSceneSection;
+class UMovieSceneSkeletalAnimationTrack;
+
+/**
+ * Linked animation track provider for the Animation Mixer.
+ * When a binding has a UMovieSceneAnimationMixerTrack, this provider
+ * creates the linked anim track as a child track inside the mixer
+ * and handles isolation by toggling mixer tracks rather than individual tracks.
+ */
+class FMixerLinkedAnimTrackProvider : public IMovieSceneLinkedAnimTrackProvider
+{
+public:
+	FMixerLinkedAnimTrackProvider(UE::Sequencer::FAnimationMixerTrackEditor* InTrackEditor)
+		: TrackEditor(InTrackEditor)
+	{}
+
+	//~ IMovieSceneLinkedAnimTrackProvider interface
+	virtual int32 GetLinkedAnimTrackPriority() const override { return 50; }
+	virtual bool CanHandleBinding(const TWeakPtr<ISequencer>& InSequencer, FGuid InBindingID) const override;
+	virtual UMovieSceneTrack* GetLinkedAnimTrack(const TWeakPtr<ISequencer>& InSequencer, FGuid InInBindingID) const override;
+	virtual UMovieSceneTrack* GetOrCreateLinkedAnimTrack(const TWeakPtr<ISequencer>& InSequencer, FGuid InBindingID, UE::Sequencer::FCommonAnimationTrackEditor& TrackEditor) override;
+	virtual bool IsolateLinkedAnimTrack(const TWeakPtr<ISequencer>& InSequencer, FGuid InBindingID, bool bIsolate) override;
+	virtual void UpdateLinkedAnimSectionRange(const TWeakPtr<ISequencer>& InSequencer, FGuid InBindingID) override;
+
+private:
+	/** Whether or not it has a mixer track*/
+	bool HasMixerTrack(const TWeakPtr<ISequencer>& InSequencer, FGuid InInBindingID) const;
+	/** Find the MixerTrack on the binding and associated section, if any. It may return a mixer track with no section, which means
+	we still need to create the section*/
+	TPair<UMovieSceneAnimationMixerTrack*,UMovieSceneSection*> FindMixerTrack(const TWeakPtr<ISequencer>& InSequencer, FGuid InBindingID) const;
+	/** Find the child skeletal animation track inside a MixerTrack that contains the linked anim sequence. */
+	UMovieSceneSkeletalAnimationTrack* FindLinkedChildTrack(
+		UMovieSceneAnimationMixerTrack* MixerTrack,
+		const TArray<struct FLevelSequenceAnimSequenceLinkItem*>& LinkedItems) const;
+
+	UE::Sequencer::FAnimationMixerTrackEditor* TrackEditor = nullptr;
+};
diff --git a/Engine/Plugins/Experimental/MovieSceneAnimMixer/Source/MovieSceneAnimMixerEditor/Private/MovieSceneAnimationMixerTrackEditor.cpp b/Engine/Plugins/Experimental/MovieSceneAnimMixer/Source/MovieSceneAnimMixerEditor/Private/MovieSceneAnimationMixerTrackEditor.cpp
--- a/Engine/Plugins/Experimental/MovieSceneAnimMixer/Source/MovieSceneAnimMixerEditor/Private/MovieSceneAnimationMixerTrackEditor.cpp
+++ b/Engine/Plugins/Experimental/MovieSceneAnimMixer/Source/MovieSceneAnimMixerEditor/Private/MovieSceneAnimationMixerTrackEditor.cpp
@@ -1,6 +1,8 @@
 // Copyright Epic Games, Inc. All Rights Reserved.
 
 #include "MovieSceneAnimationMixerTrackEditor.h"
+#include "MixerLinkedAnimTrackProvider.h"
+#include "IMovieSceneLinkedAnimTrackProvider.h"
 
 #include "AnimMixerTrailHierarchy.h"
 #include "InteractiveToolManager.h"
@@ -458,6 +460,11 @@ void FAnimationMixerTrackEditor::OnInitialize()
 	// Register as animation mixer item menu provider
 	IModularFeatures::Get().RegisterModularFeature(IMovieSceneAnimMixerItemMenuProvider::GetModularFeatureName(), this);
 
+	// Register as linked anim track provider for mixer bindings
+	MixerLinkedAnimProvider = MakeUnique<FMixerLinkedAnimTrackProvider>(this);
+	IModularFeatures::Get().RegisterModularFeature(
+		IMovieSceneLinkedAnimTrackProvider::GetModularFeatureName(), MixerLinkedAnimProvider.Get());
+
 	// Activate the mixer edit mode for skeleton visualization and gizmo keying.
 	// ActivateMode is a no-op if already active, so safe with multiple instances.
 	++MixerEditModeRefCount;
@@ -503,6 +510,14 @@ void FAnimationMixerTrackEditor::OnRelease()
 	// Unregister as animation mixer item menu provider
 	IModularFeatures::Get().UnregisterModularFeature(IMovieSceneAnimMixerItemMenuProvider::GetModularFeatureName(), this);
 
+	// Unregister linked anim track provider
+	if (MixerLinkedAnimProvider.IsValid())
+	{
+		IModularFeatures::Get().UnregisterModularFeature(
+			IMovieSceneLinkedAnimTrackProvider::GetModularFeatureName(), MixerLinkedAnimProvider.Get());
+		MixerLinkedAnimProvider.Reset();
+	}
+
 	// Unregister trail callbacks
 	UMotionTrailTool::OnCreateAdditionalTrailHierarchies.Remove(AdditionalTrailHierarchiesHandle);
 	MixerTrailHierarchy.Reset();
diff --git a/Engine/Plugins/Experimental/MovieSceneAnimMixer/Source/MovieSceneAnimMixerEditor/Private/MovieSceneAnimationMixerTrackEditor.h b/Engine/Plugins/Experimental/MovieSceneAnimMixer/Source/MovieSceneAnimMixerEditor/Private/MovieSceneAnimationMixerTrackEditor.h
--- a/Engine/Plugins/Experimental/MovieSceneAnimMixer/Source/MovieSceneAnimMixerEditor/Private/MovieSceneAnimationMixerTrackEditor.h
+++ b/Engine/Plugins/Experimental/MovieSceneAnimMixer/Source/MovieSceneAnimMixerEditor/Private/MovieSceneAnimationMixerTrackEditor.h
@@ -10,6 +10,7 @@
 #include "StructUtils/InstancedStruct.h"
 #include "TrackEditors/CommonAnimationTrackEditor.h"
 #include "IMovieSceneAnimMixerItemMenuProvider.h"
+#include "MixerLinkedAnimTrackProvider.h"
 
 struct FTimeToPixel;
 
@@ -121,6 +122,8 @@ private:
 
 private:
 	static TMap<const UClass*, FOnMakeSectionInterfaceDelegate> MakeSectionInterfaceCallbacks;
+private:
+	TUniquePtr<FMixerLinkedAnimTrackProvider> MixerLinkedAnimProvider;
 };
 
 
diff --git a/Engine/Source/Editor/MovieSceneTools/MovieSceneTools.Build.cs b/Engine/Source/Editor/MovieSceneTools/MovieSceneTools.Build.cs
--- a/Engine/Source/Editor/MovieSceneTools/MovieSceneTools.Build.cs
+++ b/Engine/Source/Editor/MovieSceneTools/MovieSceneTools.Build.cs
@@ -21,7 +21,8 @@ public class MovieSceneTools : ModuleRules
 				"CoreUObject",
 				"Engine",
                 "InputCore",
-                "MovieSceneCapture",
+				"LevelSequence",
+				"MovieSceneCapture",
 				"UnrealEd",
 				"Sequencer",
                 "EditorWidgets",
@@ -42,8 +43,8 @@ public class MovieSceneTools : ModuleRules
                 "DesktopPlatform",
                 "Json",
                 "EditorFramework",
-                "JsonUtilities",
-				"LevelSequence",
+				"EditorSubsystem",
+				"JsonUtilities",
                 "LiveLinkInterface",
                 "MessageLog",
 				"MovieScene",
diff --git a/Engine/Source/Editor/MovieSceneTools/Private/Baker/AnimSequenceBaker.cpp b/Engine/Source/Editor/MovieSceneTools/Private/Baker/AnimSequenceBaker.cpp
new file mode 100644
--- /dev/null
+++ b/Engine/Source/Editor/MovieSceneTools/Private/Baker/AnimSequenceBaker.cpp
@@ -0,0 +1,319 @@
+// Copyright Epic Games, Inc. All Rights Reserved.
+
+#include "Baker/AnimSequenceBaker.h"
+#include "ConstraintsManager.h"
+#include "ConstraintChannel.h"
+#include "Constraints/TransformConstraintChannelInterface.h"
+#include "GameFramework/Actor.h"
+#include "ISequencer.h"
+#include "MovieScene.h"
+#include "Misc/FrameRate.h"
+#include "Components/SkeletalMeshComponent.h"
+#include "TrackEditors/CommonAnimationTrackEditor.h"
+#include "Transform/TransformConstraint.h"
+#include "Transform/TransformableHandle.h"
+
+
+int32 UE::AIE::FBakeTimeRange::CalculateIndex(const FFrameNumber& InCurrentFrame) const
+{
+	if (InCurrentFrame > StartFrame && InCurrentFrame <= EndFrame)
+	{
+		return ((InCurrentFrame.Value - StartFrame.Value) / (FrameStep.Value));
+	}
+	else if (InCurrentFrame <= StartFrame)
+	{
+		return 0;
+	}
+	else if (InCurrentFrame > EndFrame)
+	{
+		return NumFrames - 1;
+	}
+	return INDEX_NONE;
+}
+
+namespace UE::Sequencer
+{ 
+	
+
+TSharedPtr<ISequencerBakeRecorder> FAnimSequenceBakeRecorder::CreateRecorderAndAddToBaker(TSharedPtr<FSequencerBaker>& InBaker, FGuid InBindingID, UAnimSequence* InAnimSequence,
+	FLevelSequenceAnimSequenceLinkItem& InItem)
+{
+
+	TSharedPtr<ISequencerBakeRecorder> Recorder = MakeShared<UE::Sequencer::FAnimSequenceBakeRecorder>(InBaker, InBindingID, InAnimSequence, InItem);
+
+	UE::AIE::ISequencerBaker::FRecorderOptions RecorderOptions;
+	RecorderOptions.WarmupFrames = InItem.WarmUpFrames;
+	RecorderOptions.DelayBeforeStart = InItem.DelayBeforeStart;
+
+	InBaker->AddRecorder(Recorder, RecorderOptions);
+	return Recorder;
+}
+
+uint32 FAnimSequenceBakeRecorder::GetHash()
+{
+	return  GetTypeHash(BindingID) + GetTypeHash(WeakAnimSequence);
+}
+
+void FAnimSequenceBakeRecorder::BakeStarted(const UE::AIE::FBakeTimeRange& InRange)
+{
+	if (TSharedPtr<FSequencerBaker> Baker = WeakBaker.Pin())
+	{
+		if (LinkItem.bUseCustomTimeRange)
+		{
+			if (ISequencer* Sequencer = Baker->GetWeakSequencer().Pin().Get())
+			{
+				TRange<FFrameNumber> NewRange{ 0 };
+
+				const FFrameRate TickResolution = Sequencer->GetFocusedTickResolution();
+				const FFrameRate DisplayResolution = Sequencer->GetFocusedDisplayRate();
+				const FFrameNumber StartFrameInTick = FFrameRate::TransformTime(FFrameTime(LinkItem.CustomStartFrame), DisplayResolution, TickResolution).FloorToFrame();
+				FFrameNumber EndFrameInTick = FFrameRate::TransformTime(FFrameTime(LinkItem.CustomEndFrame), DisplayResolution, TickResolution).CeilToFrame();
+				if (EndFrameInTick < StartFrameInTick)
+				{
+					EndFrameInTick = StartFrameInTick;
+				}
+				StartIndex = InRange.CalculateIndex(StartFrameInTick);
+				EndIndex = InRange.CalculateIndex(EndFrameInTick);
+			}
+			else
+			{
+				StartIndex = 0;
+				EndIndex = InRange.CalculateIndex(InRange.EndFrame);
+			}
+		}
+		else
+		{
+			StartIndex = 0;
+			EndIndex = InRange.CalculateIndex(InRange.EndFrame);
+		}
+
+		InitRecorder(InRange);
+
+		if (ISequencer* Sequencer = Baker->GetWeakSequencer().Pin().Get())
+		{
+			if (USkeletalMeshComponent* SkelMeshComp = WeakComponent.Get())
+			{
+				TArray< USkeletalMeshComponent*> SkelMeshComps;
+				if (LinkItem.bEvaluateAllSkeletalMeshComponents)
+				{
+					AActor* Actor = SkelMeshComp->GetTypedOuter<AActor>();
+					if (Actor)
+					{
+						Actor->GetComponents(SkelMeshComps, false);
+					}
+				}
+				else
+				{
+					SkelMeshComps.Add(SkelMeshComp);
+				}
+				Baker->AddSkelMeshCompsToTick(SkelMeshComps);
+			}
+		}
+	}
+}
+
+void FAnimSequenceBakeRecorder::InitRecorder(const UE::AIE::FBakeTimeRange& InRange)
+{
+	if (BindingID.IsValid() && WeakAnimSequence.Get() && WeakBaker.IsValid())
+	{
+		WeakComponent = FCommonAnimationTrackEditor::AcquireSkeletalMeshFromObjectGuid(BindingID, WeakBaker.Pin()->GetWeakSequencer().Pin());
+		if (WeakComponent.IsValid())
+		{
+
+			FAnimationRecordingSettings RecordingSettings;
+			RecordingSettings.SampleFrameRate = InRange.DisplayRate;
+			RecordingSettings.Interpolation = LinkItem.Interpolation;
+			RecordingSettings.InterpMode = LinkItem.CurveInterpolation;
+			if (LinkItem.CurveInterpolation == ERichCurveInterpMode::RCIM_Constant ||
+				LinkItem.CurveInterpolation == ERichCurveInterpMode::RCIM_None)
+			{
+				RecordingSettings.TangentMode = ERichCurveTangentMode::RCTM_None;
+			}
+			else
+			{
+				RecordingSettings.TangentMode = ERichCurveTangentMode::RCTM_Auto;
+			}
+			RecordingSettings.Length = 0;
+			RecordingSettings.bRemoveRootAnimation = false;
+			RecordingSettings.bCheckDeltaTimeAtBeginning = false;
+			RecordingSettings.bSetRetargetSourceAsset = LinkItem.bSetRetargetSourceAsset;
+			RecordingSettings.bRecordTransforms = LinkItem.bExportTransforms;
+			RecordingSettings.bRecordMorphTargets = LinkItem.bExportMorphTargets;
+			RecordingSettings.bRecordAttributeCurves = LinkItem.bExportAttributeCurves;
+			RecordingSettings.bRecordMaterialCurves = LinkItem.bExportMaterialCurves;
+			RecordingSettings.bSkipCurvesWithZeroValue = LinkItem.bSkipCurvesWithZeroValue;
+			RecordingSettings.bRecordInWorldSpace = LinkItem.bRecordInWorldSpace;
+			RecordingSettings.IncludeAnimationNames = LinkItem.IncludeAnimationNames;
+			RecordingSettings.ExcludeAnimationNames = LinkItem.ExcludeAnimationNames;
+			RecordingSettings.bTransactRecording = false;
+			USkeletalMeshComponent* SkelMeshComp = WeakComponent.Get();
+			UAnimSequence* AnimSequence = WeakAnimSequence.Get();
+			AnimationRecorder.Init(SkelMeshComp, AnimSequence, nullptr, RecordingSettings);
+		}
+	}
+}
+
+void FAnimSequenceBakeRecorder::BakeFrame(const UE::AIE::FBakeTimeIndex& InIndex)
+{
+	if (InIndex.Index >= StartIndex && InIndex.Index <= EndIndex)
+	{
+		if (InIndex.Index == StartIndex)
+		{
+			AnimationRecorder.BeginRecording();
+		}
+		else
+		{
+			AnimationRecorder.Update(InIndex.DeltaTime);
+		}
+	}
+}
+
+void FAnimSequenceBakeRecorder::BakeCancelled()
+{
+	AnimationRecorder.FinishRecording();
+}
+
+void FAnimSequenceBakeRecorder::BakeFinished()
+{
+	AnimationRecorder.FinishRecording(true);
+
+	if (TSharedPtr<FSequencerBaker> Baker = WeakBaker.Pin())
+	{
+		if (TSharedPtr<ISequencer> SequencerPtr = Baker->GetWeakSequencer().Pin())
+		{
+			FCommonAnimationTrackEditor::UpdateLinkedAnimSectionRange(SequencerPtr, BindingID);
+		}
+	}
+}
+
+bool FAnimSequenceBakeRecorder::HasSequencerBinding(FGuid InGuid)
+{
+	return (InGuid == BindingID);
+}
+
+bool FAnimSequenceBakeRecorder::IsolateBakeResult(bool bIsolate) 
+{
+	bool bWasChanged = false;
+	if (TSharedPtr<FSequencerBaker> Baker = WeakBaker.Pin())
+	{
+		if (TSharedPtr<ISequencer> SequencerPtr = Baker->GetWeakSequencer().Pin())
+		{
+			if (UMovieScene* MovieScene = SequencerPtr->GetFocusedMovieSceneSequence()->GetMovieScene())
+			{
+				if (FMovieSceneBinding* Binding = MovieScene->FindBinding(BindingID))
+				{
+					if (FCommonAnimationTrackEditor::IsolateLinkedAnimTrack(SequencerPtr, BindingID, bIsolate))
+					{
+						bWasChanged = true;
+					}
+				}
+			}
+		}
+	}
+	return bWasChanged;
+}
+
+void FAnimSequenceBakeRecorder::SetRecordingEnabled(bool bInEnabled) 
+{
+	bRecordingEnabled = bInEnabled;
+}
+static FMovieSceneBinding GetBindingFromTrack(UMovieScene* InMovieScene, UMovieSceneTrack* Track)
+{
+	const TArray<FMovieSceneBinding>& Bindings = ((const UMovieScene*)InMovieScene)->GetBindings();
+	for (const FMovieSceneBinding& Binding : Bindings)
+	{
+		for (UMovieSceneTrack* PossibleTrack : Binding.GetTracks())
+		{
+			if (PossibleTrack == Track)
+			{
+				return Binding;
+			}
+		}
+	}
+	FMovieSceneBinding EmptyBinding;
+	return EmptyBinding;
+}
+
+void FAnimSequenceBakeRecorder::AddDependencies(UMovieSceneTrack* InTrack, const TSharedPtr<ISequencer>& SequencerPtr)
+{
+	if(Dependencies.Contains(InTrack) == false)
+	{ 
+		Dependencies.Add(InTrack);
+		for (UMovieSceneSection* Section : InTrack->GetAllSections())
+		{
+			if (IMovieSceneConstrainedSection* ConstrainedSection = Cast<IMovieSceneConstrainedSection>(Section))
+			{
+				for (FConstraintAndActiveChannel& ConstraintChannel : ConstrainedSection->GetConstraintsChannels())
+				{
+					if (UTickableTransformConstraint* TransformConstraint = Cast< UTickableTransformConstraint>(ConstraintChannel.GetConstraint()))
+					{
+						const FConstraintChannelInterfaceRegistry& InterfaceRegistry = FConstraintChannelInterfaceRegistry::Get();
+						ITransformConstraintChannelInterface* ParentInterface = InterfaceRegistry.FindConstraintChannelInterface(TransformConstraint->ParentTRSHandle->GetClass());
+						if (ParentInterface)
+						{
+							if (UMovieSceneSection* ParentSection = ParentInterface->GetHandleSection(TransformConstraint->ParentTRSHandle, SequencerPtr->AsShared()))
+							{
+								if (UMovieSceneTrack* ParentTrack = ParentSection->GetTypedOuter<UMovieSceneTrack>())
+								{
+									if (UMovieScene* MovieScene = SequencerPtr->GetFocusedMovieSceneSequence()->GetMovieScene())
+									{
+										FMovieSceneBinding Binding = GetBindingFromTrack(MovieScene, ParentTrack);
+										if (Binding.GetObjectGuid().IsValid())
+										{
+											for (UMovieSceneTrack* PossibleTrack : Binding.GetTracks())
+											{
+												if (PossibleTrack)
+												{
+													AddDependencies(PossibleTrack, SequencerPtr);
+												}
+											}
+										}
+									}
+								}
+							}
+						}
+					}
+				}
+			}
+		}
+	}
+}
+
+void FAnimSequenceBakeRecorder::ReadyToBake()
+{
+	if (TSharedPtr<FSequencerBaker> Baker = WeakBaker.Pin())
+	{
+		if (TSharedPtr<ISequencer> SequencerPtr = Baker->GetWeakSequencer().Pin())
+		{
+			if (UMovieScene* MovieScene = SequencerPtr->GetFocusedMovieSceneSequence()->GetMovieScene())
+			{
+				if (FMovieSceneBinding* Binding = MovieScene->FindBinding(BindingID))
+				{
+					for (UMovieSceneTrack* PossibleTrack : Binding->GetTracks())
+					{
+						if (PossibleTrack)
+						{
+							AddDependencies(PossibleTrack, SequencerPtr);
+						}
+					}
+				}
+			}
+		}
+		TSharedPtr<ISequencerBakeRecorder> Shared = SharedThis(this);
+		TArray<TWeakObjectPtr<UMovieSceneSignedObject>> DependenciesArray = Dependencies.Array();
+		Baker->AddSignedObjectsToTrack(DependenciesArray, Shared);
+	}
+}
+
+void FAnimSequenceBakeRecorder::RemovedFromBake() 
+{
+	if (TSharedPtr<FSequencerBaker> Baker = WeakBaker.Pin())
+	{
+		TSharedPtr<ISequencerBakeRecorder> Shared = SharedThis(this);
+		TArray<TWeakObjectPtr<UMovieSceneSignedObject>> DependenciesArray = Dependencies.Array();
+		Baker->RemoveSignedObjectsFromTrack(DependenciesArray, Shared);
+	}
+}
+
+}
+
diff --git a/Engine/Source/Editor/MovieSceneTools/Private/Baker/AnimSequenceBaker.h b/Engine/Source/Editor/MovieSceneTools/Private/Baker/AnimSequenceBaker.h
new file mode 100644
--- /dev/null
+++ b/Engine/Source/Editor/MovieSceneTools/Private/Baker/AnimSequenceBaker.h
@@ -0,0 +1,91 @@
+// Copyright Epic Games, Inc. All Rights Reserved.
+
+#pragma once
+
+#include "AnimationRecorder.h"
+#include "Animation/AnimSequence.h"
+#include "Baker/ISequencerBaker.h"
+#include "Baker/SequencerBaker.h"
+#include "Components/SkeletalMeshComponent.h"
+#include "LevelSequenceAnimSequenceLink.h"
+#include "Misc/Guid.h"
+
+#define UE_API MOVIESCENETOOLS_API
+
+namespace UE::Sequencer
+{ 
+
+/**
+*  A Bake Recorder that creates linked anim sequences
+*/
+
+class FAnimSequenceBakeRecorder
+	: public ISequencerBakeRecorder, public TSharedFromThis<FAnimSequenceBakeRecorder>
+{
+public:
+	FAnimSequenceBakeRecorder()
+		: WeakBaker(nullptr)
+		, WeakComponent(nullptr)
+		, WeakAnimSequence(nullptr) 
+	{};
+	FAnimSequenceBakeRecorder(TSharedPtr<FSequencerBaker>& InBaker, FGuid InBindingID, UAnimSequence* InAnimSequencer, 
+		FLevelSequenceAnimSequenceLinkItem& InItem)
+		: WeakBaker(InBaker)
+		, BindingID(InBindingID)
+		, WeakComponent(nullptr)
+		, WeakAnimSequence(InAnimSequencer)
+		, LinkItem(InItem)
+	{};
+	virtual ~FAnimSequenceBakeRecorder() {};
+
+	//ISequencerBakeRecorder
+	virtual void BakeStarted(const UE::AIE::FBakeTimeRange& InRange) override;
+	virtual void BakeFrame(const UE::AIE::FBakeTimeIndex& InIndex) override;
+	virtual void BakeCancelled() override;
+	virtual void BakeFinished() override;
+	virtual bool HasSequencerBinding(FGuid InGuid) override;
+	virtual bool IsolateBakeResult(bool bIsolate) override;
+	virtual void SetRecordingEnabled(bool bInEnabled) override;
+	virtual bool GetRecordingEnabled() override
+	{
+		return bRecordingEnabled;
+	}
+	virtual uint32 GetHash() override;
+	virtual void ReadyToBake() override;
+	virtual void RemovedFromBake() override;
+
+	UE_API static TSharedPtr<ISequencerBakeRecorder> CreateRecorderAndAddToBaker(TSharedPtr<FSequencerBaker>& InBaker, FGuid InBindingID, UAnimSequence* InAnimSequence,
+		FLevelSequenceAnimSequenceLinkItem& InItem);
+private:
+
+	//may want to turn it off
+	bool bRecordingEnabled = true;
+
+	TWeakPtr<FSequencerBaker> WeakBaker;
+	//id in sequencer of skelmeshcomp we are baking
+	FGuid BindingID; 
+	TWeakObjectPtr<USkeletalMeshComponent> WeakComponent; 
+	TWeakObjectPtr<UAnimSequence> WeakAnimSequence;
+
+	//list of tracks we are dependent upon
+	TSet<TWeakObjectPtr<UMovieSceneSignedObject>> Dependencies; 
+
+	int32 StartIndex = INDEX_NONE;
+	int32 EndIndex = INDEX_NONE;
+
+	FLevelSequenceAnimSequenceLinkItem LinkItem; // containts baking options
+
+	FAnimRecorderInstance AnimationRecorder;
+
+private:
+	void InitRecorder(const UE::AIE::FBakeTimeRange& InRange);
+	//add dependencies recursively
+	void AddDependencies(UMovieSceneTrack* InTrack, const TSharedPtr<ISequencer>& SequencerPtr);
+
+};
+
+}
+
+#undef UE_API
+
+
diff --git a/Engine/Source/Editor/MovieSceneTools/Private/Baker/SequencerBaker.cpp b/Engine/Source/Editor/MovieSceneTools/Private/Baker/SequencerBaker.cpp
new file mode 100644
--- /dev/null
+++ b/Engine/Source/Editor/MovieSceneTools/Private/Baker/SequencerBaker.cpp
@@ -0,0 +1,783 @@
+// Copyright Epic Games, Inc. All Rights Reserved.
+
+#include "Baker/SequencerBaker.h"
+#include "Baker/AnimSequenceBaker.h"
+#include "ConstraintsManager.h"
+#include "Editor.h"
+#include "Evaluation/MovieSceneEvaluationTemplateInstance.h"
+#include "IMovieScenePlayer.h"
+#include "ISequencer.h"
+#include "LevelEditorViewport.h"
+#include "LevelSequence.h"
+#include "MovieScene.h"
+#include "MovieSceneSequence.h"
+#include "MovieSceneToolsModule.h"
+#include "Transform/AnimationEvaluation.h"
+#include "TrackEditors/CommonAnimationTrackEditor.h"
+
+
+bool UE::AIE::FBakeTimeRange::UpdateIfNeeded(const TRange<FFrameNumber>& InRange, const FFrameRate& InDisplayRate, const FFrameRate& InTickResolution)
+{
+	if (InRange != IncomingRange || InDisplayRate != DisplayRate || InTickResolution != TickResolution)
+	{
+		IncomingRange = InRange;
+		TickResolution = InTickResolution;
+		DisplayRate = InDisplayRate;
+		FrameStep = FFrameRate::TransformTime(FFrameNumber(1), InDisplayRate,
+			InTickResolution).RoundToFrame();
+		StartFrame = InRange.GetLowerBoundValue();
+		EndFrame = InRange.GetUpperBoundValue();
+		//make sure end frame falls on frame step 
+		NumFrames = (EndFrame.Value - StartFrame.Value) / (FrameStep.Value) + 1;
+		EndFrame.Value = StartFrame.Value + (NumFrames * FrameStep.Value) - 1;
+		return true;
+	}
+	return false;
+}
+
+namespace UE::Sequencer
+{
+
+FSequencerBaker::FSequencerBaker() 
+	: bNeedsBaking(false)
+	, bIsBaking(false)
+{	
+}
+
+UE::AIE::FBakeTimeRange FSequencerBaker::GetTimeRange()
+{
+	return BakeInterval.TimeRange;
+}
+
+bool FSequencerBaker::IsBakeRunning(TOptional<float>& OutPercentageDone)
+{
+	return bIsBaking && !BakeInterval.IsDoneCalculating();
+}
+
+void FSequencerBaker::CancelBake()
+{
+	BakeCancelled();
+}
+
+void FSequencerBaker::Initialize(TSharedPtr<ISequencer>& InSequencer)
+{
+	Release();
+	if (InSequencer)
+	{
+		WeakSequencer = InSequencer;
+		OnPreNavigateToSequenceChangedHandle = InSequencer->OnPreNavigateToSequence().AddRaw(this, &FSequencerBaker::OnPreNavigateToSequenceChanged);
+		OnActivateSequenceChangedHandle = InSequencer->OnActivateSequence().AddRaw(this, &FSequencerBaker::OnActivateSequenceChanged);
+		OnMovieSceneBindingsChangedHandle = InSequencer->OnMovieSceneBindingsChanged().AddRaw(this, &FSequencerBaker::OnMovieSceneBindingsChanged);
+		BakeInterval.TimeRange = UE::AIE::FBakeTimeRange();
+		BakeInterval.SetFromSequencer(InSequencer.Get());
+		bNeedsBaking = true;
+		OnActivateSequenceChanged(InSequencer->GetFocusedTemplateID());
+	}
+}
+
+void FSequencerBaker::Release()
+{
+	if (TSharedPtr<ISequencer> Sequencer = WeakSequencer.Pin())
+	{
+		BakeCancelled();
+		if (OnPreNavigateToSequenceChangedHandle.IsValid())
+		{
+			Sequencer->OnPreNavigateToSequence().Remove(OnPreNavigateToSequenceChangedHandle);
+			OnPreNavigateToSequenceChangedHandle.Reset();
+		}
+		if (OnActivateSequenceChangedHandle.IsValid())
+		{
+			Sequencer->OnActivateSequence().Remove(OnActivateSequenceChangedHandle);
+			OnActivateSequenceChangedHandle.Reset();
+		}
+		if (OnMovieSceneBindingsChangedHandle.IsValid())
+		{
+			Sequencer->OnMovieSceneBindingsChanged().Remove(OnMovieSceneBindingsChangedHandle);
+			OnMovieSceneBindingsChangedHandle.Reset();
+		}
+		for (TPair<uint32, TSharedPtr < ISequencerBakeRecorder>>& Pair : Recorders)
+		{
+			if (Pair.Value.IsValid())
+			{
+				Pair.Value.Reset();
+			}
+		}
+		Recorders.Reset();
+	}
+	WeakSequencer.Reset();
+}
+
+void FSequencerBaker::OnMovieSceneBindingsChanged()
+{
+	//MZ todo, need to do what we did for anim layers and find bad bindings and replace them with good ones
+}
+
+void FSequencerBaker::OnPreNavigateToSequenceChanged()
+{
+	if (TSharedPtr<ISequencer> Sequencer = WeakSequencer.Pin())
+	{
+		if (bIsBaking)
+		{
+			BakeFramesTillDone();
+		}
+		TArray<TSharedPtr<ISequencerBakeRecorder>> RecordersAsArray;
+		Recorders.GenerateValueArray(RecordersAsArray);
+		IsolateBakedTracks(RecordersAsArray, true);
+	}
+}
+
+void FSequencerBaker::SetUpBakeRecorders()
+{
+	TSharedPtr<ISequencer> SequencerPtr = WeakSequencer.Pin();
+	if (!SequencerPtr.IsValid())
+	{
+		return;
+	}
+	ULevelSequence* LevelSequence = Cast<ULevelSequence>(SequencerPtr->GetFocusedMovieSceneSequence());
+	if (LevelSequence && LevelSequence->GetClass()->ImplementsInterface(UInterface_AssetUserData::StaticClass()))
+	{
+		if (IInterface_AssetUserData* AssetUserDataInterface = Cast< IInterface_AssetUserData >(LevelSequence))
+		{
+			ULevelSequenceAnimSequenceLink* LevelAnimLink = AssetUserDataInterface->GetAssetUserData< ULevelSequenceAnimSequenceLink >();
+			if (LevelAnimLink && LevelAnimLink->AnimSequenceLinks.Num() > 0)
+			{
+				for (int32 Index = LevelAnimLink->AnimSequenceLinks.Num() - 1; Index >= 0; --Index)
+				{
+					FLevelSequenceAnimSequenceLinkItem& Item = LevelAnimLink->AnimSequenceLinks[Index];
+					if (Item.bAutoBake)
+					{
+						UAnimSequence* AnimSequence = Item.ResolveAnimSequence();
+						if (AnimSequence == nullptr)
+						{
+							LevelAnimLink->AnimSequenceLinks.RemoveAt(Index);
+							continue;
+						}
+
+						TSharedPtr<FSequencerBaker> Shared = SharedThis(this);
+						FAnimSequenceBakeRecorder::CreateRecorderAndAddToBaker(Shared, Item.SkelTrackGuid, AnimSequence, Item);
+					}
+				}
+			}
+		}
+	}
+}
+
+void FSequencerBaker::OnActivateSequenceChanged(FMovieSceneSequenceIDRef ID) 
+{
+	if (!GIsTransacting)//is also called from undo from sequencer!!!!
+	{
+		if (TSharedPtr<ISequencer> Sequencer = WeakSequencer.Pin())
+		{
+			BakeInterval.TimeRange = UE::AIE::FBakeTimeRange();
+			BakeInterval.SetFromSequencer(Sequencer.Get());
+			BakeCancelled();
+			BakingRecorders.Reset();
+			GEditor->GetTimerManager()->SetTimerForNextTick([this, WeakThis = this->AsWeak()]()
+			{
+					if (WeakThis.IsValid())
+					{
+						if (TSharedPtr<ISequencer> Sequencer = WeakSequencer.Pin())
+						{
+							SetUpBakeRecorders();
+						}
+					}
+			});
+		}
+	}
+}
+
+bool FSequencerBaker::CheckForChanges()
+{
+	bool bHasChange = false;
+	if (TSharedPtr<ISequencer> Sequencer = WeakSequencer.Pin())
+	{
+		if (Sequencer->GetFocusedMovieSceneSequence() && Sequencer->GetFocusedMovieSceneSequence()->GetMovieScene())
+		{
+			bHasChange = LastValidMovieSceneGuid != Sequencer->GetFocusedMovieSceneSequence()->GetMovieScene()->GetSignature();
+			LastValidMovieSceneGuid = Sequencer->GetFocusedMovieSceneSequence()->GetMovieScene()->GetSignature();
+		}
+	}
+	return bHasChange;
+}
+
+void FSequencerBaker::Tick(float DeltaTime)
+{
+	if (TSharedPtr<ISequencer> Sequencer = WeakSequencer.Pin())
+	{
+		if (BakeInterval.SetFromSequencer(Sequencer.Get()))
+		{
+			bNeedsBaking = true;
+			AddAllRecordersToBake();
+			return;
+		}
+	}
+
+	if (CheckForChanges())
+	{
+		for (TPair<uint32, TSharedPtr < ISequencerBakeRecorder>>& Pair : Recorders)
+		{
+			if (Pair.Value.IsValid())
+			{
+				Pair.Value->ReadyToBake();
+			}
+		}
+		TArray<TSharedPtr<ISequencerBakeRecorder>> ChangedRecorders = SignedObjectsToTrack.UpdateSignatures();
+		if (ChangedRecorders.Num() > 0)
+		{
+			BakingRecorders.Append(ChangedRecorders);
+			bNeedsBaking = true;
+		}
+	}
+}
+
+void FSequencerBaker::BakeStarted()
+{
+	if (TSharedPtr<ISequencer> SequencePtr = WeakSequencer.Pin())
+	{
+		BakeInterval.SetFromSequencer(SequencePtr.Get());
+
+		//cache out important objects
+		SkeletalMeshCompsTick.Reset();
+		LastBakedRecorders.Reset();
+
+		FMovieSceneSequenceTransform RootToLocalTransform = SequencePtr->GetFocusedMovieSceneSequenceTransform();
+		LocalToRootTransform = RootToLocalTransform.Inverse();
+		World = GCurrentLevelEditingViewportClient ? GCurrentLevelEditingViewportClient->GetWorld() : nullptr;
+
+		if (bIsBaking)
+		{
+			BakeCancelled(); // cancel the bake if we were baking
+		}
+
+		if (BakingRecorders.Num() > 0)
+		{
+			bNeedsBaking = false;
+			bIsBaking = true;
+			if (SequencePtr->GetFocusedMovieSceneSequence() && SequencePtr->GetFocusedMovieSceneSequence()->GetMovieScene())
+			{
+				UMovieScene* MovieScene = SequencePtr->GetFocusedMovieSceneSequence()->GetMovieScene();
+				for (TSharedPtr < ISequencerBakeRecorder>& Recorder : BakingRecorders)
+				{
+					if (Recorder.IsValid())
+					{
+						Recorder->BakeStarted(BakeInterval.TimeRange);
+					}
+				}
+				for (IMovieSceneToolsAnimationBakeHelper* BakeHelper : BakeHelpers)
+				{
+					if (BakeHelper)
+					{
+						BakeHelper->StartBaking(MovieScene);
+					}
+				}
+			}
+			const float DeltaTime =  1.0 / BakeInterval.TimeRange.DisplayRate.AsDecimal();;
+			//if we have warmup frames
+			if (BakeInterval.WarmupFrames > 0)
+			{
+				for (int32 Index = BakeInterval.WarmupFrames.Value; Index > 0; --Index)
+				{
+					FFrameNumber FrameNumber = BakeInterval.TimeRange.StartFrame - FFrameNumber(Index);
+					TickFrameInternal(FrameNumber, DeltaTime);
+				}
+			}
+			//if we have delay frames run them first at the LocalStartFrame - ExportOptions->WarmupFrames;
+			if (BakeInterval.DelayBeforeStart > 0)
+			{
+				for (int32 Index = 0; Index < BakeInterval.DelayBeforeStart.Value; ++Index)
+				{
+					TickFrameInternal(BakeInterval.TimeRange.StartFrame, DeltaTime);
+				}
+			}
+		}
+	}
+	else
+	{
+		BakeCancelled();
+	}
+}
+
+void FSequencerBaker::AddAllRecordersToBake()
+{
+	for (TPair<uint32, TSharedPtr < ISequencerBakeRecorder>>& Pair : Recorders)
+	{
+		if (Pair.Value.IsValid())
+		{
+			BakingRecorders.Add(Pair.Value);
+		}
+	}
+}
+
+void FSequencerBaker::BakeFramesTillDone()
+{
+	while (bIsBaking && BakingRecorders.Num() > 0)
+	{
+		int32 NextFrame = BakeInterval.NextFrameToCalculate();
+		if (NextFrame != INDEX_NONE)
+		{
+			UE::AIE::FBakeTimeIndex BakeTimeIndex;
+			BakeTimeIndex.Index = NextFrame;
+			BakeTimeIndex.FrameNumber = BakeInterval.TimeRange.StartFrame + (BakeInterval.TimeRange.FrameStep * BakeTimeIndex.Index);
+			BakeTimeIndex.DeltaTime = 1.0 / BakeInterval.TimeRange.DisplayRate.AsDecimal();
+			TickFrameInternal(BakeTimeIndex.FrameNumber, BakeTimeIndex.DeltaTime);
+			for (TSharedPtr < ISequencerBakeRecorder>& Recorder : BakingRecorders)
+			{
+				if (Recorder.IsValid())
+				{
+					Recorder->BakeFrame(BakeTimeIndex);
+				}
+			}
+			BakeInterval.SetFrameCalculated(NextFrame);
+		}
+		else
+		{
+			BakeFinished(); //sets bIsBaking == false
+		}
+	}
+}
+
+void FSequencerBaker::BakeFullRange()
+{
+	if (TSharedPtr<ISequencer> SequencePtr = WeakSequencer.Pin())
+	{
+		if (bIsBaking)
+		{
+			CancelBake();
+		}
+
+		if (BakingRecorders.Num() == 0)
+		{
+			return;
+		}
+
+		BakeStarted();
+		BakeFramesTillDone();
+
+		BakingRecorders.Reset();
+		SequencePtr->ForceEvaluate();
+	}
+}
+
+bool FSequencerBaker::TickFrameInternal(const FFrameNumber& FrameNumber, float DeltaTime)
+{
+	TSharedPtr<ISequencer> Sequencer = WeakSequencer.Pin();
+	if (Sequencer && Sequencer->GetFocusedMovieSceneSequence() && Sequencer->GetFocusedMovieSceneSequence()->GetMovieScene())
+	{
+
+		UMovieScene* MovieScene = Sequencer->GetFocusedMovieSceneSequence()->GetMovieScene();
+		//Begin records a frame so need to set things up first
+		for (IMovieSceneToolsAnimationBakeHelper* BakeHelper : BakeHelpers)
+		{
+			if (BakeHelper)
+			{
+				BakeHelper->PreEvaluation(MovieScene, FrameNumber);
+			}
+		}
+		FFrameTime GlobalTime(FrameNumber);
+		GlobalTime = LocalToRootTransform.TryTransformTime(GlobalTime).Get(GlobalTime); //player evals in root time so need to go back to it.
+
+		FMovieSceneContext MovieSceneContext = FMovieSceneContext(FMovieSceneEvaluationRange(GlobalTime, BakeInterval.TimeRange.TickResolution), Sequencer->GetPlaybackStatus()).SetHasJumped(true);
+		Sequencer->GetEvaluationTemplate().EvaluateSynchronousBlocking(MovieSceneContext);
+
+		if (World)
+		{
+			const UE::Anim::FEvaluationForCachingScope CachingScope(DeltaTime);
+			const FConstraintsManagerController& Controller = FConstraintsManagerController::Get(World);
+			Controller.EvaluateAllConstraints();
+		}
+
+		for (IMovieSceneToolsAnimationBakeHelper* BakeHelper : BakeHelpers)
+		{
+			if (BakeHelper)
+			{
+				BakeHelper->PostEvaluation(MovieScene, FrameNumber);
+			}
+		}
+		//actors to tick? todo
+		// 
+		// Update space bases so new animation position has an effect.
+		for (TWeakObjectPtr<USkeletalMeshComponent>& SkelMeshComp : SkeletalMeshCompsTick)
+		{
+			if (SkelMeshComp.IsValid())
+			{
+				//SkelMeshComp->UpdateLODStatus(); 
+
+				SkelMeshComp->SetForcedLOD(1); //we need to do this every tick since a restore state may happen while evaluating and bummp us out ot a different LOD
+				SkelMeshComp->TickAnimation(DeltaTime, false);
+
+				SkelMeshComp->RefreshBoneTransforms();
+				SkelMeshComp->RefreshFollowerComponents();
+				SkelMeshComp->UpdateComponentToWorld();
+				SkelMeshComp->FinalizeBoneTransform();
+				SkelMeshComp->MarkRenderTransformDirty();
+				SkelMeshComp->MarkRenderDynamicDataDirty();
+			}
+		}
+		return true;
+	}
+	return false;
+}
+
+void FSequencerBaker::BakeFrame()
+{
+	if (bIsBaking && BakingRecorders.Num() > 0)
+	{
+		int32 NextFrame = BakeInterval.NextFrameToCalculate();
+		if (NextFrame != INDEX_NONE)
+		{
+			UE::AIE::FBakeTimeIndex BakeTimeIndex;
+			BakeTimeIndex.Index = NextFrame;
+			BakeTimeIndex.FrameNumber = BakeInterval.TimeRange.StartFrame + (BakeInterval.TimeRange.FrameStep * BakeTimeIndex.Index);
+			BakeTimeIndex.DeltaTime = 1.0 / BakeInterval.TimeRange.DisplayRate.AsDecimal();
+			TickFrameInternal(BakeTimeIndex.FrameNumber, BakeTimeIndex.DeltaTime);
+			for (TSharedPtr < ISequencerBakeRecorder>& Recorder : BakingRecorders)
+			{
+				if (Recorder.IsValid())
+				{
+					Recorder->BakeFrame(BakeTimeIndex);
+				}
+			}
+			BakeInterval.SetFrameCalculated(NextFrame);
+		}
+		else
+		{
+			BakeFinished();
+		}
+	}
+	else //check if baking
+	{
+		if (bIsBaking)
+		{
+			BakeCancelled();
+		}
+	}
+}
+
+void FSequencerBaker::BakeCancelled()
+{
+	if (bIsBaking)
+	{
+		bIsBaking = false;
+		BakeInterval.Clear();
+		for (TSharedPtr < ISequencerBakeRecorder>& Recorder : BakingRecorders)
+		{
+			if (Recorder.IsValid())
+			{
+				Recorder->BakeCancelled();
+			}
+		}
+		SkeletalMeshCompsTick.Reset();
+	}
+}
+
+void FSequencerBaker::BakeFinished()
+{
+	//if (bIsBaking && Recorders.Num() > 0)
+	if (bIsBaking)
+	{
+		bIsBaking = false;
+		bNeedsBaking = false;
+		BakeInterval.Clear();
+		for (TSharedPtr < ISequencerBakeRecorder>& Recorder : BakingRecorders)
+		{
+			if (Recorder.IsValid())
+			{
+				Recorder->BakeFinished();
+				LastBakedRecorders.Add(Recorder);
+			}
+		}
+		BakingRecorders.Reset();
+		SkeletalMeshCompsTick.Reset();
+	}
+}
+
+void FSequencerBaker::IsolateBakedTracks(TArray<TSharedPtr<ISequencerBakeRecorder>>& InRecorders, bool bIsolate)
+{
+	for (TSharedPtr < ISequencerBakeRecorder>& Recorder : InRecorders)
+	{
+		if (Recorder.IsValid())
+		{
+			Recorder->IsolateBakeResult(bIsolate);
+		}
+	}
+}
+
+void FSequencerBaker::AddSkelMeshCompsToTick(TArray<USkeletalMeshComponent*>& InSkelMeshCompArray)
+{
+	for (USkeletalMeshComponent* SkelmeshComp : InSkelMeshCompArray)
+	{
+		if (SkelmeshComp)
+		{
+			SkeletalMeshCompsTick.Add(SkelmeshComp);
+		}
+	}
+}
+
+void FSequencerBaker::IgnoreLastChange()
+{
+	SignedObjectsToTrack.UpdateSignatures();
+}
+
+void FSequencerBaker::AddSignedObjectsToTrack(TArray<TWeakObjectPtr<UMovieSceneSignedObject>>& InSignedObjects, TSharedPtr<ISequencerBakeRecorder>& InOwner)
+{
+	if (InOwner.IsValid())
+	{
+		for (TWeakObjectPtr< UMovieSceneSignedObject>& SignedObject : InSignedObjects)
+		{
+			if (SignedObject.IsValid())
+			{
+				SignedObjectsToTrack.AddObject(SignedObject.Get(), InOwner);
+			}
+		}
+	}
+}
+
+void FSequencerBaker::RemoveSignedObjectsFromTrack(TArray<TWeakObjectPtr<UMovieSceneSignedObject>>& InSignedObjects, TSharedPtr<ISequencerBakeRecorder>& InOwner)
+{
+	if (InOwner.IsValid())
+	{
+		for (TWeakObjectPtr< UMovieSceneSignedObject>& SignedObject : InSignedObjects)
+		{
+			if (SignedObject.IsValid())
+			{
+				SignedObjectsToTrack.RemoveObject(SignedObject.Get(), InOwner);
+			}
+		}
+	}
+}
+
+void FSequencerBaker::AddExtraRecordersToBake(const TArray<TSharedPtr<ISequencerBakeRecorder>>& InExtraRecorders)
+{
+	for (const TSharedPtr< ISequencerBakeRecorder>& Extra : InExtraRecorders)
+	{
+		if (Extra.IsValid())
+		{
+			BakingRecorders.Add(Extra);
+			bNeedsBaking = true;
+		}
+	}
+}
+
+bool FSequencerBaker::AddRecorderToBake(TSharedPtr<ISequencerBakeRecorder>& InRecorder)
+{
+	if (Recorders.Contains(InRecorder->GetHash()) && BakingRecorders.Contains(InRecorder) == false)
+	{
+		BakingRecorders.Add(InRecorder);
+		bNeedsBaking = true;
+		return true;
+	}
+	return false;
+}
+
+void FSequencerBaker::AddRecorder(TSharedPtr<ISequencerBakeRecorder>& InRecorder,const UE::AIE::ISequencerBaker::FRecorderOptions& Options)
+{
+	if (InRecorder.IsValid())
+	{
+		Recorders.Add(InRecorder->GetHash(),InRecorder);
+		InRecorder->ReadyToBake();
+		//use largest specified delay and warmup
+		if (Options.DelayBeforeStart > BakeInterval.DelayBeforeStart)
+		{
+			BakeInterval.DelayBeforeStart = Options.DelayBeforeStart;
+		}
+		if (Options.WarmupFrames > BakeInterval.WarmupFrames)
+		{
+			BakeInterval.WarmupFrames = Options.WarmupFrames;
+		}
+		OnAddBakeRecorder.Broadcast(InRecorder.Get());
+	}
+}
+
+void FSequencerBaker::RemoveRecorder(TSharedPtr<ISequencerBakeRecorder>& InRecorder)
+{
+	if (InRecorder.IsValid())
+	{
+		Recorders.Remove(InRecorder->GetHash());
+		OnRemoveBakeRecorder.Broadcast(InRecorder.Get());
+	}
+	if (BakingRecorders.Contains(InRecorder))
+	{
+		if (InRecorder.IsValid())
+		{
+			InRecorder->BakeCancelled();
+			InRecorder->RemovedFromBake();
+		}
+		BakingRecorders.Remove(InRecorder);
+		if (BakingRecorders.Num() == 0)
+		{
+			CancelBake();
+		}
+	}
+}
+
+void FSequencerBaker::IsolateRecorderWithGuid(FGuid InGuid, bool bIsolate)
+{
+	for (TPair<uint32, TSharedPtr < ISequencerBakeRecorder>>& Pair : Recorders)
+	{
+		if (Pair.Value.IsValid())
+		{
+			if (Pair.Value->HasSequencerBinding(InGuid))
+			{
+				Pair.Value->IsolateBakeResult(bIsolate);
+			}
+		}
+	}
+}
+
+bool FSequencerBaker::FBakeInterval::SetFromSequencer(ISequencer* InSequencer)
+{	
+	if (InSequencer && InSequencer->GetFocusedMovieSceneSequence() && InSequencer->GetFocusedMovieSceneSequence()->GetMovieScene())
+	{
+		FFrameRate TickResolution = InSequencer->GetFocusedTickResolution();
+		FFrameRate DisplayRate = InSequencer->GetFocusedDisplayRate();
+		TOptional<TRange<FFrameNumber>> OptionalRange = InSequencer->GetSubSequenceRange();
+		TRange<FFrameNumber> Range = OptionalRange.IsSet() ? OptionalRange.GetValue() : InSequencer->GetFocusedMovieSceneSequence()->GetMovieScene()->GetPlaybackRange();
+		if (TimeRange.UpdateIfNeeded(Range, DisplayRate, TickResolution))
+		{
+			CalculatedFrames.Reset();
+			CalculatedFrames.Init(false, TimeRange.NumFrames);
+			return true;
+		}
+	}
+	return false;
+}
+
+void FSequencerBaker::FBakeInterval::Clear()
+{
+	CalculatedFrames.Init(false, TimeRange.NumFrames);
+}
+
+bool FSequencerBaker::FBakeInterval::IsFrameCalcuated(int32 Index)
+{
+	if (Index >= 0 && Index < TimeRange.NumFrames)
+	{
+		return CalculatedFrames[Index];
+	}
+	return false;
+}
+
+void FSequencerBaker::FBakeInterval::SetFrameCalculated(int32 Index)
+{
+	if (Index >= 0 && Index < TimeRange.NumFrames)
+	{
+		CalculatedFrames[Index] = true;
+	}
+}
+
+int32 FSequencerBaker::FBakeInterval::NextFrameToCalculate()
+{
+	return CalculatedFrames.Find(false);
+}
+
+bool FSequencerBaker::FBakeInterval::IsDoneCalculating() const
+{
+	return (TimeRange.NumFrames > 0 && CalculatedFrames.IsEmpty() == false &&
+		CalculatedFrames.Find(false) == INDEX_NONE);
+}
+
+static TArray<bool> GetDisabledArray(UMovieSceneTrack* Track)
+{
+	TArray<bool> IsDisabled;
+	if (Track)
+	{
+		if (Track->GetMaxRowIndex() == 0)
+		{
+			IsDisabled.Add(Track->IsEvalDisabled());
+		}
+		else
+		{
+			for (int32 RowIndex = 0; RowIndex <= Track->GetMaxRowIndex(); ++RowIndex)
+			{
+				IsDisabled.Add(Track->IsRowEvalDisabled(RowIndex));
+			}
+		}
+	}
+	return IsDisabled;
+}
+
+void FSequencerBaker::FSignedObjectChanges::AddObject(UMovieSceneSignedObject* InObject, TSharedPtr<ISequencerBakeRecorder>& Owner)
+{
+	if (InObject)
+	{
+		FGuid Signature = InObject->GetSignature();
+		TArray<bool> DisabledArray;
+		if (UMovieSceneTrack* Track = Cast<UMovieSceneTrack>(InObject))
+		{
+			DisabledArray = GetDisabledArray(Track);
+		}
+		if (FSignedObjectState* AlreadyHere = ObjectSignatures.Find(InObject))
+		{
+			//added track with new owner, clear guid so it get's baked
+			if(AlreadyHere->Owners.Contains(Owner) == false)
+			{
+				AlreadyHere->Guid.Invalidate();
+				AlreadyHere->Owners.Add(Owner);
+			}
+		}
+		else //track was never added, clear guid so it get's baked.
+		{
+			FSignedObjectState State;
+			State.Guid.Invalidate();
+			State.DisabledArray = DisabledArray;
+			State.Owners.Add(Owner);
+			ObjectSignatures.Add(InObject, State);
+		}
+	}
+}
+
+void FSequencerBaker::FSignedObjectChanges::RemoveObject(UMovieSceneSignedObject* InObject, TSharedPtr<ISequencerBakeRecorder>& Owner)
+{
+	if (InObject)
+	{
+		if (FSignedObjectState* AlreadyHere = ObjectSignatures.Find(InObject))
+		{
+			AlreadyHere->Owners.Remove(Owner);
+			if (AlreadyHere->Owners.Num() == 0)
+			{
+				ObjectSignatures.Remove(InObject);
+			}
+		}
+	}
+}
+
+TArray<TSharedPtr<ISequencerBakeRecorder>> FSequencerBaker::FSignedObjectChanges::UpdateSignatures()
+{
+	TArray<TSharedPtr<ISequencerBakeRecorder>> ChangedRecorders;
+	for (TPair<TWeakObjectPtr<UMovieSceneSignedObject>, FSignedObjectState>& Pair : ObjectSignatures)
+	{
+		if (Pair.Key.IsValid())
+		{
+			FGuid Signature = Pair.Key->GetSignature();
+
+			if (Signature != Pair.Value.Guid)
+			{
+				TArray<bool> DisabledArray;
+				if (UMovieSceneTrack* Track = Cast<UMovieSceneTrack>(Pair.Key.Get()))
+				{
+					DisabledArray = GetDisabledArray(Track);
+				}
+				if (IgnoredObjects.Contains(Pair.Key) == false) //if ignored update sig but 
+				{
+					if (UMovieSceneTrack* Track = Cast<UMovieSceneTrack>(Pair.Key.Get()))
+					{
+						DisabledArray = GetDisabledArray(Track);
+					}
+					if (DisabledArray.Num() == 0 || DisabledArray == Pair.Value.DisabledArray)
+					{
+						for (TSharedPtr < ISequencerBakeRecorder>& Owner : Pair.Value.Owners)
+						{
+							ChangedRecorders.Add(Owner);
+						}
+					}
+				}
+				ObjectSignatures[Pair.Key].DisabledArray = DisabledArray;
+				ObjectSignatures[Pair.Key].Guid = Signature;
+			}
+		}
+	}
+	return ChangedRecorders;
+}
+
+};
+
+
diff --git a/Engine/Source/Editor/MovieSceneTools/Private/Baker/SequencerBakerSubsystem.cpp b/Engine/Source/Editor/MovieSceneTools/Private/Baker/SequencerBakerSubsystem.cpp
new file mode 100644
--- /dev/null
+++ b/Engine/Source/Editor/MovieSceneTools/Private/Baker/SequencerBakerSubsystem.cpp
@@ -0,0 +1,83 @@
+// Copyright Epic Games, Inc. All Rights Reserved.
+
+
+#include "Baker/SequencerBakerSubsystem.h"
+#include "Baker/SequencerBaker.h"
+#include "ISequencerModule.h"
+
+#include "Evaluation/MovieScenePlayback.h"
+#include "LevelSequence.h"
+#include "ISceneOutliner.h"
+#include "MovieScenePossessable.h"
+#include "SequenceBindingTree.h"
+#include "MovieScene.h"
+
+#include "MovieSceneSpawnable.h"
+
+void USequencerBakeSubsystem::Initialize(FSubsystemCollectionBase& Collection)
+{
+	Super::Initialize(Collection);
+	UE_LOG(LogMovieScene, Log, TEXT("USequencerBakeSubsystem subsystem initialized."));
+
+	ISequencerModule& SequencerModule = FModuleManager::Get().LoadModuleChecked<ISequencerModule>("Sequencer");
+	OnSequencerCreatedHandle = SequencerModule.RegisterOnSequencerCreated(FOnSequencerCreated::FDelegate::CreateUObject(this, &USequencerBakeSubsystem::OnSequencerCreated));
+}
+
+void USequencerBakeSubsystem::Deinitialize()
+{
+	for (TPair<TWeakPtr<ISequencer>, TSharedPtr<UE::Sequencer::FSequencerBaker>>& Pair: Sequencers)
+	{
+		if (Pair.Value.IsValid())
+		{
+			Pair.Value->Release();
+		}
+
+	}
+	
+	Sequencers.Reset();
+
+	Super::Deinitialize();
+
+	UE_LOG(LogMovieScene, Log, TEXT("USequencerBakeSubsystem subsystem deinitialized."));
+
+	ISequencerModule* SequencerModulePtr = FModuleManager::Get().GetModulePtr<ISequencerModule>("Sequencer");
+	if (SequencerModulePtr)
+	{
+		SequencerModulePtr->UnregisterOnSequencerCreated(OnSequencerCreatedHandle);
+	}
+}
+
+void USequencerBakeSubsystem::OnSequencerCreated(TSharedRef<ISequencer> InSequencer)
+{
+	UE_LOG(LogMovieScene, VeryVerbose, TEXT("ULevelSequenceEditorSubsystem::OnSequencerCreated"));
+
+	TSharedPtr<UE::Sequencer::FSequencerBaker> BakerPtr = MakeShared<UE::Sequencer::FSequencerBaker>();
+	TSharedPtr<ISequencer> SequencerPtr = InSequencer.ToSharedPtr();
+	BakerPtr->Initialize(SequencerPtr);
+	Sequencers.Add(TWeakPtr<ISequencer>(InSequencer), BakerPtr);
+	
+	InSequencer->OnCloseEvent().AddUObject(this, &USequencerBakeSubsystem::OnSequencerClosed);
+}
+
+void USequencerBakeSubsystem::OnSequencerClosed(TSharedRef<ISequencer> InSequencer)
+{
+	UE_LOG(LogMovieScene, VeryVerbose, TEXT("ULevelSequenceEditorSubsystem::OnSequencerClosed"));
+
+	InSequencer.Get().OnCloseEvent().RemoveAll(this);
+
+	if (TSharedPtr<UE::Sequencer::FSequencerBaker>* BakerPtr = Sequencers.Find(TWeakPtr<ISequencer>(InSequencer)))
+	{
+		(*BakerPtr)->Release();
+		Sequencers.Remove(TWeakPtr<ISequencer>(InSequencer));
+	}
+}
+
+TSharedPtr<UE::Sequencer::FSequencerBaker> USequencerBakeSubsystem::GetSequencerBaker(const TSharedPtr<ISequencer>& InSequencer)
+{
+	TSharedPtr<UE::Sequencer::FSequencerBaker> SequencerPtr = nullptr;
+	if (TSharedPtr<UE::Sequencer::FSequencerBaker>* SequencerPtrPtr = Sequencers.Find(TWeakPtr<ISequencer>(InSequencer)))
+	{
+		SequencerPtr = *SequencerPtrPtr;
+	}
+	return SequencerPtr;
+}
diff --git a/Engine/Source/Editor/MovieSceneTools/Private/IMovieSceneLinkedAnimTrackProvider.cpp b/Engine/Source/Editor/MovieSceneTools/Private/IMovieSceneLinkedAnimTrackProvider.cpp
new file mode 100644
--- /dev/null
+++ b/Engine/Source/Editor/MovieSceneTools/Private/IMovieSceneLinkedAnimTrackProvider.cpp
@@ -0,0 +1,9 @@
+// Copyright Epic Games, Inc. All Rights Reserved.
+
+#include "IMovieSceneLinkedAnimTrackProvider.h"
+
+FName IMovieSceneLinkedAnimTrackProvider::GetModularFeatureName()
+{
+	static FName FeatureName = FName(TEXT("MovieSceneLinkedAnimTrackProvider"));
+	return FeatureName;
+}
diff --git a/Engine/Source/Editor/MovieSceneTools/Private/TrackEditors/CommonAnimationTrackEditor.cpp b/Engine/Source/Editor/MovieSceneTools/Private/TrackEditors/CommonAnimationTrackEditor.cpp
--- a/Engine/Source/Editor/MovieSceneTools/Private/TrackEditors/CommonAnimationTrackEditor.cpp
+++ b/Engine/Source/Editor/MovieSceneTools/Private/TrackEditors/CommonAnimationTrackEditor.cpp
@@ -1,8 +1,13 @@
 // Copyright Epic Games, Inc. All Rights Reserved.
 
 #include "TrackEditors/CommonAnimationTrackEditor.h"
+#include "TrackEditors/DefaultLinkedAnimTrackProvider.h"
 #include "EditModes/SkeletalAnimationTrackEditMode.h"
 #include "Tracks/MovieSceneCommonAnimationTrack.h"
+#include "Tracks/MovieSceneSkeletalAnimationTrack.h"
+#include "Tracks/MovieScene3DAttachTrack.h"
+#include "Tracks/MovieScene3DPathTrack.h"
+#include "Tracks/MovieScene3DTransformTrack.h"
 #include "Rendering/DrawElements.h"
 #include "Widgets/SBoxPanel.h"
 #include "Framework/MultiBox/MultiBoxBuilder.h"
@@ -56,6 +61,9 @@
 #include "Engine/Blueprint.h"
 
 #include "Subsystems/AssetEditorSubsystem.h"
+
+#include "Baker/SequencerBakerSubsystem.h"
+#include "Baker/AnimSequenceBaker.h"
 #include "Editor.h"
 
 #include "ClassViewerFilter.h"
@@ -81,7 +89,12 @@
 #include "Widgets/Layout/SBox.h"
 #include "Widgets/Layout/SUniformGridPanel.h"
 #include "Widgets/Input/SButton.h"
+#include "Widgets/Input/SCheckBox.h"
 #include "Widgets/Input/SComboBox.h"
+#include "Widgets/Images/SImage.h"
+#include "MVVM/ViewModels/OutlinerColumns/OutlinerColumnTypes.h"
+#include "MVVM/ViewModels/ObjectBindingModel.h"
+#include "MVVM/Selection/Selection.h"
 #include "Exporters/AnimSeqExportOption.h"
 
 #include "EditModes/SkeletalAnimationTrackEditMode.h"
@@ -105,6 +118,8 @@ namespace UE::Sequencer
 
 int32 FCommonAnimationTrackEditor::NumberActive = 0;
 
+static FDefaultLinkedAnimTrackProvider GDefaultLinkedAnimTrackProvider;
+
 namespace CommonAnimationEditorConstants
 {
 	// @todo Sequencer Allow this to be customizable
@@ -1269,6 +1284,48 @@ void FCommonAnimationSection::BuildSectionContextMenu(FMenuBuilder& MenuBuilder,
 	MenuBuilder.EndSection();
 }
 
+bool FCommonAnimationTrackEditor::AcquireAllBoundSkeletalMeshObjectsFromObjectGuid(const FGuid& Guid, TSharedPtr<ISequencer> SequencerPtr, TArray<USkeletalMeshComponent*>& OutSkelMeshComponents,
+	TArray<FGuid>& OutBindingIds)
+{
+	UObject* BoundObject = SequencerPtr.IsValid() ? SequencerPtr->FindSpawnedObjectOrTemplate(Guid) : nullptr;
+	if (AActor* Actor = Cast<AActor>(BoundObject))
+	{
+		TArray<USkeletalMeshComponent*> ActorSkelMeshComps;
+		Actor->GetComponents(ActorSkelMeshComps);
+
+		if (ActorSkelMeshComps.Num() == 1)
+		{
+			OutSkelMeshComponents.Add(ActorSkelMeshComps[0]);
+			OutBindingIds.Add(Guid);
+		}
+		else
+		{
+			constexpr bool bCreateHandleIfMissing = false;
+			for (USkeletalMeshComponent* Component : ActorSkelMeshComps)
+			{
+				if (Component)
+				{
+					FGuid ObjectHandle = SequencerPtr->GetHandleToObject(Component, bCreateHandleIfMissing);
+					if (ObjectHandle.IsValid())
+					{
+						OutSkelMeshComponents.Add(Component);
+						OutBindingIds.Add(ObjectHandle);
+					}
+				}
+			}
+		}
+	}
+	else if (USkeletalMeshComponent* SkeletalMeshComponent = Cast<USkeletalMeshComponent>(BoundObject))
+	{
+		if (SkeletalMeshComponent->GetSkeletalMeshAsset())
+		{
+			OutSkelMeshComponents.Add(SkeletalMeshComponent);
+			OutBindingIds.Add(Guid);
+		}
+	}
+	return OutSkelMeshComponents.Num() > 0;
+}
+
 USkeletalMeshComponent* FCommonAnimationTrackEditor::AcquireSkeletalMeshFromObjectGuid(const FGuid& Guid, TSharedPtr<ISequencer> SequencerPtr)
 {
 	UObject* BoundObject = SequencerPtr.IsValid() ? SequencerPtr->FindSpawnedObjectOrTemplate(Guid) : nullptr;
@@ -1407,6 +1464,12 @@ void FCommonAnimationTrackEditor::OnInitialize()
 
 	++FCommonAnimationTrackEditor::NumberActive;
 
+	if (FCommonAnimationTrackEditor::NumberActive == 1)
+	{
+		IModularFeatures::Get().RegisterModularFeature(
+			IMovieSceneLinkedAnimTrackProvider::GetModularFeatureName(), &GDefaultLinkedAnimTrackProvider);
+	}
+
 	if (ShouldActivateSkeletalAnimEditMode())
 	{
 		// Activate the default mode in case FEditorModeTools::Tick isn't run before here.
@@ -1441,6 +1504,9 @@ void FCommonAnimationTrackEditor::OnRelease()
 	}
 	if (FCommonAnimationTrackEditor::NumberActive == 0)
 	{
+		IModularFeatures::Get().UnregisterModularFeature(
+			IMovieSceneLinkedAnimTrackProvider::GetModularFeatureName(), &GDefaultLinkedAnimTrackProvider);
+
 		GLevelEditorModeTools().DeactivateMode(FSkeletalAnimationTrackEditMode::ModeName);
 	}
 }
@@ -1504,29 +1570,211 @@ void FCommonAnimationTrackEditor::BuildObjectBindingContextMenu(FMenuBuilder& Me
 		ConstructObjectBindingTrackMenu(MenuBuilder, ObjectBindings);
 	}
 }
-void FCommonAnimationTrackEditor::OnSequencerSaved(ISequencer& )
+
+void FCommonAnimationTrackEditor::SaveAnimSequenceLink(FLevelSequenceAnimSequenceLinkItem& Item)
 {
 	TSharedPtr<ISequencer> SequencerPtr = GetSequencer();
 	if (!SequencerPtr.IsValid())
 	{
 		return;
 	}
+	UAnimSequence* AnimSequence = Item.ResolveAnimSequence();
+	if (AnimSequence == nullptr)
+	{
+		return;
+	}
 	ULevelSequence* LevelSequence = Cast<ULevelSequence>(SequencerPtr->GetFocusedMovieSceneSequence());
-	if (LevelSequence && LevelSequence->GetClass()->ImplementsInterface(UInterface_AssetUserData::StaticClass()))
+	UMovieSceneSequence* MovieSceneSequence = SequencerPtr->GetFocusedMovieSceneSequence();
+	UMovieSceneSequence* RootMovieSceneSequence = SequencerPtr->GetRootMovieSceneSequence();
+	FMovieSceneSequenceIDRef Template = SequencerPtr->GetFocusedTemplateID();
+	FMovieSceneSequenceTransform RootToLocalTransform = SequencerPtr->GetFocusedMovieSceneSequenceTransformUnwarped();
+
+	if (IInterface_AssetUserData* AnimAssetUserData = Cast< IInterface_AssetUserData >(AnimSequence))
 	{
-		if (IInterface_AssetUserData* AssetUserDataInterface = Cast< IInterface_AssetUserData >(LevelSequence))
+		UAnimSequenceLevelSequenceLink* AnimLevelLink = AnimAssetUserData->GetAssetUserData< UAnimSequenceLevelSequenceLink >();
+		if (!AnimLevelLink)
 		{
-			ULevelSequenceAnimSequenceLink* LevelAnimLink = AssetUserDataInterface->GetAssetUserData< ULevelSequenceAnimSequenceLink >();
-			if (LevelAnimLink && LevelAnimLink->AnimSequenceLinks.Num() > 0)
+			AnimLevelLink = NewObject<UAnimSequenceLevelSequenceLink>(AnimSequence, NAME_None, RF_Public | RF_Transactional);
+			AnimAssetUserData->AddAssetUserData(AnimLevelLink);
+		}
+		AnimLevelLink->SetLevelSequence(LevelSequence);
+		AnimLevelLink->SkelTrackGuid = Item.SkelTrackGuid;
+	}
+	USkeletalMeshComponent* SkelMeshComp = UE::Sequencer::FCommonAnimationTrackEditor::AcquireSkeletalMeshFromObjectGuid(Item.SkelTrackGuid, GetSequencer());
+	if (AnimSequence && SkelMeshComp && AnimSeqExportOption)
+	{
+		const bool bSavedSetRetargetSourceAsset = AnimSeqExportOption->bSetRetargetSourceAsset;
+		const bool bSavedExportMorphTargets = AnimSeqExportOption->bExportMorphTargets;
+		const bool bSavedExportAttributeCurves = AnimSeqExportOption->bExportAttributeCurves;
+		const bool bSavedExportMaterialCurves = AnimSeqExportOption->bExportMaterialCurves;
+		const bool bSavedSkipCurvesWithZeroValue = AnimSeqExportOption->bSkipCurvesWithZeroValue;
+		const bool bSavedExportTransforms = AnimSeqExportOption->bExportTransforms;
+		const bool bSavedIncludeComponentTransform = AnimSeqExportOption->bRecordInWorldSpace;
+		const bool bSavedEvaluateAllSkeletalMeshComponents = AnimSeqExportOption->bEvaluateAllSkeletalMeshComponents;
+		const EAnimInterpolationType SavedInterpolationType = AnimSeqExportOption->Interpolation;
+		const ERichCurveInterpMode SavedCurveInterpolationType = AnimSeqExportOption->CurveInterpolation;
+		const TArray<FString> SavedIncludeAnimationNames = AnimSeqExportOption->IncludeAnimationNames;
+		const TArray<FString> SavedExcludeAnimationNames = AnimSeqExportOption->ExcludeAnimationNames;
+		const FFrameNumber SavedWarmUpFrames = AnimSeqExportOption->WarmUpFrames;
+		const FFrameNumber SavedDelayBeforeStart = AnimSeqExportOption->DelayBeforeStart;
+		const bool bSavedUseCustomTimeRange = AnimSeqExportOption->bUseCustomTimeRange;
+		const FFrameNumber SavedCustomStartFrame = AnimSeqExportOption->CustomStartFrame;
+		const FFrameNumber SavedCustomEndFrame = AnimSeqExportOption->CustomEndFrame;
+		const FFrameRate SavedCustomDisplayRate = AnimSeqExportOption->CustomDisplayRate;
+		const bool bSavedUseCustomFrameRate = AnimSeqExportOption->bUseCustomFrameRate;
+		const FFrameRate SavedCustomFrameRate = AnimSeqExportOption->CustomFrameRate;
+
+		//If the asset cleared the retarget source asset we don't want to set it!
+		if (Item.bSetRetargetSourceAsset)
+		{
+			if (AnimSequence->GetRetargetSourceAsset().IsNull() == true)
+			{
+				AnimSeqExportOption->bSetRetargetSourceAsset = false;
+			}
+			else
 			{
+				AnimSeqExportOption->bSetRetargetSourceAsset = true;
+			}
+		}
+		else
+		{
+			AnimSeqExportOption->bSetRetargetSourceAsset = false;
+		}
+		AnimSeqExportOption->bExportMorphTargets = Item.bExportMorphTargets;
+		AnimSeqExportOption->bExportAttributeCurves = Item.bExportAttributeCurves;
+		AnimSeqExportOption->bExportMaterialCurves = Item.bExportMaterialCurves;
+		AnimSeqExportOption->bSkipCurvesWithZeroValue = Item.bSkipCurvesWithZeroValue;
+
+		AnimSeqExportOption->bExportTransforms = Item.bExportTransforms;
+		AnimSeqExportOption->bRecordInWorldSpace = Item.bRecordInWorldSpace;
+		AnimSeqExportOption->bEvaluateAllSkeletalMeshComponents = Item.bEvaluateAllSkeletalMeshComponents;
+		AnimSeqExportOption->Interpolation = Item.Interpolation;
+		AnimSeqExportOption->CurveInterpolation = Item.CurveInterpolation;
+
+		AnimSeqExportOption->IncludeAnimationNames = Item.IncludeAnimationNames;
+		AnimSeqExportOption->ExcludeAnimationNames = Item.ExcludeAnimationNames;
+		AnimSeqExportOption->WarmUpFrames = Item.WarmUpFrames;
+		AnimSeqExportOption->DelayBeforeStart = Item.DelayBeforeStart;
+		AnimSeqExportOption->bUseCustomTimeRange = Item.bUseCustomTimeRange;
+		AnimSeqExportOption->CustomStartFrame = Item.CustomStartFrame;
+		AnimSeqExportOption->CustomEndFrame = Item.CustomEndFrame;
+		AnimSeqExportOption->CustomDisplayRate = Item.CustomDisplayRate;
+		AnimSeqExportOption->bUseCustomFrameRate = Item.bUseCustomFrameRate;
+		AnimSeqExportOption->CustomFrameRate = Item.CustomFrameRate;
+
+		TGuardValue<bool> TransactionOptionGuard(AnimSeqExportOption->bTransactRecording, false);
+		FAnimExportSequenceParameters AESP;
+		AESP.Player = SequencerPtr.Get();
+		AESP.RootToLocalTransform = RootToLocalTransform;
+		AESP.MovieSceneSequence = MovieSceneSequence;
+		AESP.RootMovieSceneSequence = RootMovieSceneSequence;
+		AESP.bForceUseOfMovieScenePlaybackRange = SequencerPtr.Get()->GetSequencerSettings()->ShouldEvaluateSubSequencesInIsolation();
+		//make sure all spanwables are present
+		SequencerPtr->ForceEvaluate();
+		bool bResult = MovieSceneToolHelpers::ExportToAnimSequence(AnimSequence, AnimSeqExportOption, AESP, SkelMeshComp);
+
+		AnimSeqExportOption->bSetRetargetSourceAsset = bSavedSetRetargetSourceAsset;
+		AnimSeqExportOption->bExportMorphTargets = bSavedExportMorphTargets;
+		AnimSeqExportOption->bExportAttributeCurves = bSavedExportAttributeCurves;
+		AnimSeqExportOption->bExportMaterialCurves = bSavedExportMaterialCurves;
+		AnimSeqExportOption->bSkipCurvesWithZeroValue = bSavedSkipCurvesWithZeroValue;
+		AnimSeqExportOption->bExportTransforms = bSavedExportTransforms;
+		AnimSeqExportOption->bRecordInWorldSpace = bSavedIncludeComponentTransform;
+		AnimSeqExportOption->bEvaluateAllSkeletalMeshComponents = bSavedEvaluateAllSkeletalMeshComponents;
+		AnimSeqExportOption->Interpolation = SavedInterpolationType;
+		AnimSeqExportOption->CurveInterpolation = SavedCurveInterpolationType;
+
+		AnimSeqExportOption->IncludeAnimationNames = SavedIncludeAnimationNames;
+		AnimSeqExportOption->ExcludeAnimationNames = SavedExcludeAnimationNames;
+		AnimSeqExportOption->WarmUpFrames = SavedWarmUpFrames;
+		AnimSeqExportOption->DelayBeforeStart = SavedDelayBeforeStart;
+		AnimSeqExportOption->bUseCustomTimeRange = bSavedUseCustomTimeRange;
+		AnimSeqExportOption->CustomStartFrame = SavedCustomStartFrame;
+		AnimSeqExportOption->CustomEndFrame = SavedCustomEndFrame;
+		AnimSeqExportOption->CustomDisplayRate = SavedCustomDisplayRate;
+		AnimSeqExportOption->bUseCustomFrameRate = bSavedUseCustomFrameRate;
+		AnimSeqExportOption->CustomFrameRate = SavedCustomFrameRate;
+		//save the anim sequence to disk to make sure they are in sync
+		UPackage* const Package = AnimSequence->GetOutermost();
+		FString const PackageName = Package->GetName();
+		FString const PackageFileName = FPackageName::LongPackageNameToFilename(PackageName, FPackageName::GetAssetPackageExtension());
+
+		FSavePackageArgs SaveArgs;
+		SaveArgs.TopLevelFlags = RF_Standalone;
+		SaveArgs.SaveFlags = SAVE_NoError;
+		UPackage::SavePackage(Package, NULL, *PackageFileName, SaveArgs);
+	}
+}
 
-				UMovieSceneSequence* MovieSceneSequence = SequencerPtr->GetFocusedMovieSceneSequence();
-				UMovieSceneSequence* RootMovieSceneSequence = SequencerPtr->GetRootMovieSceneSequence();
-				FMovieSceneSequenceIDRef Template = SequencerPtr->GetFocusedTemplateID();
-				FMovieSceneSequenceTransform RootToLocalTransform = SequencerPtr->GetFocusedMovieSceneSequenceTransformUnwarped();
-				//if in subsquence and we want to, turn on ShouldEvaluateSubSequencesInIsolation, use the following object for that
+void FCommonAnimationTrackEditor::UpdateLinkedAnimSectionRange(const TWeakPtr<ISequencer>& InSequencer, FGuid BindingID)
+{
+	if (IMovieSceneLinkedAnimTrackProvider* Provider = FindProviderForBinding(InSequencer, BindingID))
+	{
+		Provider->UpdateLinkedAnimSectionRange(InSequencer, BindingID);
+	}
+}
+
+TAutoConsoleVariable<bool> CVarUseOldSaveLinkedAnims(TEXT("Sequencer.UseOldSaveLinkedAnimsWhenSaving"), false, TEXT("When true use old system when saving linked anims, if false use new, faster system (default false) "));
+
+void FCommonAnimationTrackEditor::OnSequencerSaved(ISequencer& )
+{
+	
+	TSharedPtr<ISequencer> SequencerPtr = GetSequencer();
+	if (!SequencerPtr.IsValid())
+	{
+		return;
+	}
+	if (CVarUseOldSaveLinkedAnims->GetBool())
+	{
+
+		ULevelSequence* LevelSequence = Cast<ULevelSequence>(SequencerPtr->GetFocusedMovieSceneSequence());
+		if (LevelSequence && LevelSequence->GetClass()->ImplementsInterface(UInterface_AssetUserData::StaticClass()))
+		{
+			if (IInterface_AssetUserData* AssetUserDataInterface = Cast< IInterface_AssetUserData >(LevelSequence))
+			{
+				ULevelSequenceAnimSequenceLink* LevelAnimLink = AssetUserDataInterface->GetAssetUserData< ULevelSequenceAnimSequenceLink >();
+				if (LevelAnimLink && LevelAnimLink->AnimSequenceLinks.Num() > 0)
+				{
+
+					//if in subsquence and we want to, turn on ShouldEvaluateSubSequencesInIsolation, use the following object for that
+					{
+						UE::Sequencer::FSequencerBakingSetupRestore RestoreBaking(SequencerPtr);
+						for (int32 Index = LevelAnimLink->AnimSequenceLinks.Num() - 1; Index >= 0; --Index)
+						{
+							FLevelSequenceAnimSequenceLinkItem& Item = LevelAnimLink->AnimSequenceLinks[Index];
+							UAnimSequence* AnimSequence = Item.ResolveAnimSequence();
+							if (AnimSequence == nullptr)
+							{
+								LevelAnimLink->AnimSequenceLinks.RemoveAt(Index);
+								continue;
+							}
+							SaveAnimSequenceLink(Item);
+						}
+					}//restore bake settings
+
+					//re-evaluate at current frame
+					SequencerPtr->ForceEvaluate();
+				}
+			}
+		}
+	}
+	else
+	{
+
+		ULevelSequence* LevelSequence = Cast<ULevelSequence>(SequencerPtr->GetFocusedMovieSceneSequence());
+		if (LevelSequence && LevelSequence->GetClass()->ImplementsInterface(UInterface_AssetUserData::StaticClass()))
+		{
+			USequencerBakeSubsystem* SequencerBakeSubsystem = GEditor->GetEditorSubsystem<USequencerBakeSubsystem>();
+			TSharedPtr<UE::Sequencer::FSequencerBaker> SequencerBaker = SequencerBakeSubsystem ? SequencerBakeSubsystem->GetSequencerBaker(SequencerPtr) : nullptr;
+			//ones not linked up to do smart baking, still add them to system
+			TArray<TPair<UAnimSequence*, FGuid>> TempsToSave;
+			UE::Sequencer::FSequencerBakingSetupRestore RestoreBaking(SequencerPtr);
+
+			if (IInterface_AssetUserData* AssetUserDataInterface = Cast< IInterface_AssetUserData >(LevelSequence))
+			{
+				ULevelSequenceAnimSequenceLink* LevelAnimLink = AssetUserDataInterface->GetAssetUserData< ULevelSequenceAnimSequenceLink >();
+				if (LevelAnimLink && LevelAnimLink->AnimSequenceLinks.Num() > 0)
 				{
-					UE::Sequencer::FSequencerBakingSetupRestore RestoreBaking(SequencerPtr);
 					for (int32 Index = LevelAnimLink->AnimSequenceLinks.Num() - 1; Index >= 0; --Index)
 					{
 						FLevelSequenceAnimSequenceLinkItem& Item = LevelAnimLink->AnimSequenceLinks[Index];
@@ -1536,110 +1784,47 @@ void FCommonAnimationTrackEditor::OnSequencerSaved(ISequencer& )
 							LevelAnimLink->AnimSequenceLinks.RemoveAt(Index);
 							continue;
 						}
-						if (IInterface_AssetUserData* AnimAssetUserData = Cast< IInterface_AssetUserData >(AnimSequence))
+						if (Item.bAutoBake && SequencerBaker.IsValid())
 						{
-							UAnimSequenceLevelSequenceLink* AnimLevelLink = AnimAssetUserData->GetAssetUserData< UAnimSequenceLevelSequenceLink >();
-							if (!AnimLevelLink)
+							if (SequencerBaker->FindRecorder(Item.SkelTrackGuid, AnimSequence) == nullptr)
 							{
-								AnimLevelLink = NewObject<UAnimSequenceLevelSequenceLink>(AnimSequence, NAME_None, RF_Public | RF_Transactional);
-								AnimAssetUserData->AddAssetUserData(AnimLevelLink);
+								TPair<UAnimSequence*, FGuid> Pair;
+								Pair.Key = AnimSequence;
+								Pair.Value = Item.SkelTrackGuid;
+								TempsToSave.Add(Pair);
 							}
-							AnimLevelLink->SetLevelSequence(LevelSequence);
-							AnimLevelLink->SkelTrackGuid = Item.SkelTrackGuid;
 						}
-						USkeletalMeshComponent* SkelMeshComp = UE::Sequencer::FCommonAnimationTrackEditor::AcquireSkeletalMeshFromObjectGuid(Item.SkelTrackGuid, GetSequencer());
-						if (AnimSequence && SkelMeshComp)
+						else
 						{
-							const bool bSavedSetRetargetSourceAsset = AnimSeqExportOption->bSetRetargetSourceAsset;
-							const bool bSavedExportMorphTargets = AnimSeqExportOption->bExportMorphTargets;
-							const bool bSavedExportAttributeCurves = AnimSeqExportOption->bExportAttributeCurves;
-							const bool bSavedExportMaterialCurves = AnimSeqExportOption->bExportMaterialCurves;
-							const bool bSkipCurvesWithZeroValue = AnimSeqExportOption->bSkipCurvesWithZeroValue;
-							const bool bSavedExportTransforms = AnimSeqExportOption->bExportTransforms;
-							const bool bSavedIncludeComponentTransform = AnimSeqExportOption->bRecordInWorldSpace;
-							const bool bSavedEvaluateAllSkeletalMeshComponents = AnimSeqExportOption->bEvaluateAllSkeletalMeshComponents;
-							const EAnimInterpolationType SavedInterpolationType = AnimSeqExportOption->Interpolation;
-							const ERichCurveInterpMode SavedCurveInterpolationType = AnimSeqExportOption->CurveInterpolation;
-							const TArray<FString> SavedIncludeAnimationNames = AnimSeqExportOption->IncludeAnimationNames;
-							const TArray<FString> SavedExcludeAnimationNames = AnimSeqExportOption->ExcludeAnimationNames;
-							const FFrameNumber SavedWarmUpFrames = AnimSeqExportOption->WarmUpFrames;
-							const FFrameNumber SavedDelayBeforeStart = AnimSeqExportOption->DelayBeforeStart;
-							const bool bSavedUseCustomTimeRange = AnimSeqExportOption->bUseCustomTimeRange;
-							const FFrameNumber SavedCustomStartFrame = AnimSeqExportOption->CustomStartFrame;
-							const FFrameNumber SavedCustomEndFrame = AnimSeqExportOption->CustomEndFrame;
-							const FFrameRate SavedCustomDisplayRate = AnimSeqExportOption->CustomDisplayRate;
-							const bool bSavedUseCustomFrameRate = AnimSeqExportOption->bUseCustomFrameRate;
-							const FFrameRate SavedCustomFrameRate = AnimSeqExportOption->CustomFrameRate;
-
-							//If the asset cleared the retarget source asset we don't want to set it!
-							if (Item.bSetRetargetSourceAsset)
+							//mz todo and test
+							//	TSharedPtr<ISequencerBakeRecorder> Recorder = MakeShared<UE::Sequencer::FAnimSequenceBakeRecorder>(*SequencerBaker, Item.SkelTrackGuid, AnimSequence, Item);
+							//Add to extra recordes and use that system.
+							SaveAnimSequenceLink(Item);
+						}
+					}
+					if (SequencerBaker)
+					{
+						//bake it all out
+						SequencerBaker->AddAllRecordersToBake();
+						SequencerBaker->BakeFullRange();
+
+						//now make sure it's all good
+						for (TPair<UAnimSequence*, FGuid>& Pair : TempsToSave)
+						{
+							UAnimSequence* AnimSequence = Pair.Key;
+
+							if (IInterface_AssetUserData* AnimAssetUserData = Cast< IInterface_AssetUserData >(AnimSequence))
 							{
-								if (AnimSequence->GetRetargetSourceAsset().IsNull() == true)
-								{
-									AnimSeqExportOption->bSetRetargetSourceAsset = false;
-								}
-								else
+								UAnimSequenceLevelSequenceLink* AnimLevelLink = AnimAssetUserData->GetAssetUserData< UAnimSequenceLevelSequenceLink >();
+								if (!AnimLevelLink)
 								{
-									AnimSeqExportOption->bSetRetargetSourceAsset = true;
+									AnimLevelLink = NewObject<UAnimSequenceLevelSequenceLink>(AnimSequence, NAME_None, RF_Public | RF_Transactional);
+									AnimAssetUserData->AddAssetUserData(AnimLevelLink);
 								}
+								AnimLevelLink->SetLevelSequence(LevelSequence);
+								AnimLevelLink->SkelTrackGuid = Pair.Value;
 							}
-							else
-							{
-								AnimSeqExportOption->bSetRetargetSourceAsset = false;
-							}
-							AnimSeqExportOption->bExportMorphTargets = Item.bExportMorphTargets;
-							AnimSeqExportOption->bExportAttributeCurves = Item.bExportAttributeCurves;
-							AnimSeqExportOption->bExportMaterialCurves = Item.bExportMaterialCurves;
-							AnimSeqExportOption->bSkipCurvesWithZeroValue = Item.bSkipCurvesWithZeroValue;
-							AnimSeqExportOption->bExportTransforms = Item.bExportTransforms;
-							AnimSeqExportOption->bRecordInWorldSpace = Item.bRecordInWorldSpace;
-							AnimSeqExportOption->bEvaluateAllSkeletalMeshComponents = Item.bEvaluateAllSkeletalMeshComponents;
-							AnimSeqExportOption->Interpolation = Item.Interpolation;
-							AnimSeqExportOption->CurveInterpolation = Item.CurveInterpolation;
-
-							AnimSeqExportOption->IncludeAnimationNames = Item.IncludeAnimationNames;
-							AnimSeqExportOption->ExcludeAnimationNames = Item.ExcludeAnimationNames;
-							AnimSeqExportOption->WarmUpFrames = Item.WarmUpFrames;
-							AnimSeqExportOption->DelayBeforeStart = Item.DelayBeforeStart;
-							AnimSeqExportOption->bUseCustomTimeRange = Item.bUseCustomTimeRange;
-							AnimSeqExportOption->CustomStartFrame = Item.CustomStartFrame;
-							AnimSeqExportOption->CustomEndFrame = Item.CustomEndFrame;
-							AnimSeqExportOption->CustomDisplayRate = Item.CustomDisplayRate;
-							AnimSeqExportOption->bUseCustomFrameRate = Item.bUseCustomFrameRate;
-							AnimSeqExportOption->CustomFrameRate = Item.CustomFrameRate;
-
-							TGuardValue<bool> TransactionOptionGuard(AnimSeqExportOption->bTransactRecording, false);
-							FAnimExportSequenceParameters AESP;
-							AESP.Player = SequencerPtr.Get();
-							AESP.RootToLocalTransform = RootToLocalTransform;
-							AESP.MovieSceneSequence = MovieSceneSequence;
-							AESP.RootMovieSceneSequence = RootMovieSceneSequence;
-							AESP.bForceUseOfMovieScenePlaybackRange = SequencerPtr.Get()->GetSequencerSettings()->ShouldEvaluateSubSequencesInIsolation();
-							//make sure all spanwables are present
-							SequencerPtr->ForceEvaluate();
-							bool bResult = MovieSceneToolHelpers::ExportToAnimSequence(AnimSequence, AnimSeqExportOption, AESP, SkelMeshComp);
-
-							AnimSeqExportOption->bSetRetargetSourceAsset = bSavedSetRetargetSourceAsset;
-							AnimSeqExportOption->bExportMorphTargets = bSavedExportMorphTargets;
-							AnimSeqExportOption->bExportAttributeCurves = bSavedExportAttributeCurves;
-							AnimSeqExportOption->bExportMaterialCurves = bSavedExportMaterialCurves;
-							AnimSeqExportOption->bSkipCurvesWithZeroValue = bSkipCurvesWithZeroValue;
-							AnimSeqExportOption->bExportTransforms = bSavedExportTransforms;
-							AnimSeqExportOption->bRecordInWorldSpace = bSavedIncludeComponentTransform;
-							AnimSeqExportOption->bEvaluateAllSkeletalMeshComponents = bSavedEvaluateAllSkeletalMeshComponents;
-							AnimSeqExportOption->Interpolation = SavedInterpolationType;
-							AnimSeqExportOption->CurveInterpolation = SavedCurveInterpolationType;
-
-							AnimSeqExportOption->IncludeAnimationNames = SavedIncludeAnimationNames;
-							AnimSeqExportOption->ExcludeAnimationNames = SavedExcludeAnimationNames;
-							AnimSeqExportOption->WarmUpFrames = SavedWarmUpFrames;
-							AnimSeqExportOption->DelayBeforeStart = SavedDelayBeforeStart;
-							AnimSeqExportOption->bUseCustomTimeRange = bSavedUseCustomTimeRange;
-							AnimSeqExportOption->CustomStartFrame = SavedCustomStartFrame;
-							AnimSeqExportOption->CustomEndFrame = SavedCustomEndFrame;
-							AnimSeqExportOption->CustomDisplayRate = SavedCustomDisplayRate;
-							AnimSeqExportOption->bUseCustomFrameRate = bSavedUseCustomFrameRate;
-							AnimSeqExportOption->CustomFrameRate = SavedCustomFrameRate;
+
 							//save the anim sequence to disk to make sure they are in sync
 							UPackage* const Package = AnimSequence->GetOutermost();
 							FString const PackageName = Package->GetName();
@@ -1651,10 +1836,7 @@ void FCommonAnimationTrackEditor::OnSequencerSaved(ISequencer& )
 							UPackage::SavePackage(Package, NULL, *PackageFileName, SaveArgs);
 						}
 					}
-				}//restore bake settings
-				
-				//re-evaluate at current frame
-				SequencerPtr->ForceEvaluate();
+				}
 			}
 		}
 	}
@@ -1978,6 +2160,387 @@ void FCommonAnimationTrackEditor::OpenLinkedAnimSequence(FGuid Binding)
 	}
 }
 
+IMovieSceneLinkedAnimTrackProvider* FCommonAnimationTrackEditor::FindProviderForBinding(const TWeakPtr<ISequencer>& InSequencer, FGuid BindingID)
+{
+	TArray<IMovieSceneLinkedAnimTrackProvider*> Providers =
+		IModularFeatures::Get().GetModularFeatureImplementations<IMovieSceneLinkedAnimTrackProvider>(
+			IMovieSceneLinkedAnimTrackProvider::GetModularFeatureName());
+
+	Providers.Sort([](const IMovieSceneLinkedAnimTrackProvider& A, const IMovieSceneLinkedAnimTrackProvider& B)
+	{
+		return A.GetLinkedAnimTrackPriority() < B.GetLinkedAnimTrackPriority();
+	});
+
+	for (IMovieSceneLinkedAnimTrackProvider* Provider : Providers)
+	{
+		if (Provider->CanHandleBinding(InSequencer, BindingID))
+		{
+			return Provider;
+		}
+	}
+	return nullptr;
+}
+
+UMovieSceneSkeletalAnimationTrack* FCommonAnimationTrackEditor::FindEmptyAnimationTrack(FGuid BindingID)
+{
+	TSharedPtr<ISequencer> SequencerPtr = GetSequencer();
+	if (!SequencerPtr.IsValid() || SequencerPtr->GetFocusedMovieSceneSequence() == nullptr ||
+		SequencerPtr->GetFocusedMovieSceneSequence()->GetMovieScene() == nullptr)
+	{
+		return nullptr;
+	}
+	UMovieScene* MovieScene = SequencerPtr->GetFocusedMovieSceneSequence()->GetMovieScene();
+	//todo handle multiples
+	TArray<UMovieSceneTrack*> AllTracks = MovieScene->FindTracks(UMovieSceneSkeletalAnimationTrack::StaticClass(), BindingID, NAME_None);
+	for (UMovieSceneTrack* Track : AllTracks)
+	{
+		if (UMovieSceneSkeletalAnimationTrack* AnimTrack = Cast<UMovieSceneSkeletalAnimationTrack>(Track))
+		{
+			if (AnimTrack->GetAllSections().Num() == 0)
+			{
+				return AnimTrack;
+			}
+		}
+	}
+	return nullptr;
+}
+
+void FCommonAnimationTrackEditor::CreateLinkedAnimationSequence(USkeletalMeshComponent* SkelMeshComp, USkeleton* Skeleton, FGuid Binding, bool bCreateSoftLink)
+{
+	HandleCreateAnimationSequence(SkelMeshComp, Skeleton, Binding, bCreateSoftLink);
+}
+
+TArray<FLevelSequenceAnimSequenceLinkItem*> FCommonAnimationTrackEditor::GetLinkedAnimSequences(const TWeakPtr<ISequencer>& InSequencer, FGuid Binding)
+{
+	TArray<FLevelSequenceAnimSequenceLinkItem*> LinkedItems;
+	TSharedPtr<ISequencer> SequencerPtr = InSequencer.Pin();
+	if (!SequencerPtr.IsValid())
+	{
+		return LinkedItems;
+	}
+	ULevelSequence* LevelSequence = Cast<ULevelSequence>(SequencerPtr->GetFocusedMovieSceneSequence());
+	if (LevelSequence && LevelSequence->GetClass()->ImplementsInterface(UInterface_AssetUserData::StaticClass()))
+	{
+		if (IInterface_AssetUserData* AssetUserDataInterface = Cast< IInterface_AssetUserData >(LevelSequence))
+		{
+			ULevelSequenceAnimSequenceLink* LevelAnimLink = AssetUserDataInterface->GetAssetUserData< ULevelSequenceAnimSequenceLink >();
+			if (LevelAnimLink)
+			{
+				for (FLevelSequenceAnimSequenceLinkItem& Item : LevelAnimLink->AnimSequenceLinks)
+				{
+					if (Item.SkelTrackGuid == Binding)
+					{
+						if (UAnimSequence* AnimSequence = Item.ResolveAnimSequence())
+						{
+							LinkedItems.Add(&Item);
+						}
+					}
+				}
+			}
+		}
+	}
+	return LinkedItems;
+
+}
+
+
+static bool DisableTrack(UMovieSceneTrack* Track, bool bInEvalDisabled)
+{
+	// Disable the entire track or each of the rows
+	bool bWasChanged = false;
+	if (Track->GetMaxRowIndex() == 0)
+	{
+		if (Track->IsEvalDisabled() != bInEvalDisabled)
+		{
+			bWasChanged = true;
+			Track->Modify();
+			Track->SetEvalDisabled(bInEvalDisabled);
+		}
+	}
+	else
+	{
+		for (int32 RowIndex = 0; RowIndex <= Track->GetMaxRowIndex(); ++RowIndex)
+		{
+			if (Track->IsRowEvalDisabled(RowIndex) != bInEvalDisabled)
+			{
+				bWasChanged = true;
+				Track->Modify();
+				Track->SetRowEvalDisabled(bInEvalDisabled, RowIndex);
+			}
+		}
+	}
+	return bWasChanged;
+}
+
+static bool IsObjectEnabled(UMovieSceneSignedObject* Object)
+{
+	if (UMovieSceneTrack* Track = Cast<UMovieSceneTrack>(Object))
+	{
+		if (Track->GetMaxRowIndex() == 0)
+		{
+			return !Track->IsEvalDisabled();
+		}
+		for (int32 RowIndex = 0; RowIndex <= Track->GetMaxRowIndex(); ++RowIndex)
+		{
+			if (!Track->IsRowEvalDisabled(RowIndex))
+			{
+				return true;
+			}
+		}
+
+	}
+	else if (UMovieSceneSection* Section = Cast<UMovieSceneSection>(Object))
+	{
+		if (UMovieSceneTrack* ParentTrack = Section->GetTypedOuter<UMovieSceneTrack>())
+		{
+			if (ParentTrack->GetMaxRowIndex() == 0)
+			{
+				return !ParentTrack->IsEvalDisabled();
+			}
+			else
+			{
+				const int32 SectionRowIndex = Section->GetRowIndex();
+				for (int32 RowIndex = 0; RowIndex <= ParentTrack->GetMaxRowIndex(); ++RowIndex)
+				{
+					if (SectionRowIndex == RowIndex && !ParentTrack->IsRowEvalDisabled(RowIndex))
+					{
+						return true;
+					}
+				}
+			}
+		}
+	}
+	return false;
+}
+
+bool FCommonAnimationTrackEditor::IsolateLinkedAnimTrack(const TWeakPtr<ISequencer>& InSequencer, FGuid BindingID, bool bIsolate)
+{
+	if (IMovieSceneLinkedAnimTrackProvider* Provider = FindProviderForBinding(InSequencer, BindingID))
+	{
+		return Provider->IsolateLinkedAnimTrack(InSequencer, BindingID, bIsolate);
+	}
+	return false;
+}
+
+UMovieSceneTrack* FCommonAnimationTrackEditor::GetLinkedAnimTrack(FGuid Binding, bool bCreateIfMissing)
+{
+	TSharedPtr<ISequencer> SequencerPtr = GetSequencer();
+	if (!SequencerPtr.IsValid())
+	{
+		return nullptr;
+	}
+	if (IMovieSceneLinkedAnimTrackProvider* Provider = FindProviderForBinding(SequencerPtr, Binding))
+	{
+		if (bCreateIfMissing == false)
+		{
+			return Provider->GetLinkedAnimTrack(SequencerPtr,Binding);
+		}
+		else
+		{
+			return Provider->GetOrCreateLinkedAnimTrack(SequencerPtr, Binding, *this);
+		}
+	}
+	return nullptr;
+}
+
+void FCommonAnimationTrackEditor::ConfigureAutoBake(TArray<FGuid> Bindings)
+{
+	TSharedPtr<ISequencer> SequencerPtr = GetSequencer();
+	if (!SequencerPtr.IsValid())
+	{
+		return;
+	}
+	FScopedTransaction MatchSection(LOCTEXT("ConfigureAutoBake_Transaction", "Set Up LinkedAutoBakes"));
+	for (FGuid& Binding : Bindings)
+	{
+		GetLinkedAnimTrack(Binding, true);
+		ULevelSequence* LevelSequence = Cast<ULevelSequence>(SequencerPtr->GetFocusedMovieSceneSequence());
+		if (LevelSequence && LevelSequence->GetMovieScene() && LevelSequence->GetClass()->ImplementsInterface(UInterface_AssetUserData::StaticClass()))
+		{
+			UMovieScene* MovieScene = SequencerPtr->GetFocusedMovieSceneSequence()->GetMovieScene();
+			if (IInterface_AssetUserData* AssetUserDataInterface = Cast< IInterface_AssetUserData >(LevelSequence))
+			{
+				ULevelSequenceAnimSequenceLink* LevelAnimLink = AssetUserDataInterface->GetAssetUserData< ULevelSequenceAnimSequenceLink >();
+				if (LevelAnimLink)
+				{
+					for (FLevelSequenceAnimSequenceLinkItem& Item : LevelAnimLink->AnimSequenceLinks)
+					{
+						if (Item.SkelTrackGuid == Binding && Item.bAutoBake == true)
+						{
+							UAnimSequence* AnimSequence = Item.ResolveAnimSequence();
+
+							USkeletalMeshComponent* SkelMeshComp = AcquireSkeletalMeshFromObjectGuid(Binding, SequencerPtr);
+							if (AnimSequence && SkelMeshComp)
+							{
+								if (USequencerBakeSubsystem* SequencerBakeSubsystem = GEditor->GetEditorSubsystem<USequencerBakeSubsystem>())
+								{
+									if (TSharedPtr<UE::Sequencer::FSequencerBaker> SequencerBaker = SequencerBakeSubsystem->GetSequencerBaker(SequencerPtr))
+									{
+										if (SequencerBaker->FindRecorder(Binding, AnimSequence) == nullptr)
+										{
+											TSharedPtr<ISequencerBakeRecorder> Recorder = FAnimSequenceBakeRecorder::CreateRecorderAndAddToBaker(SequencerBaker, Binding, AnimSequence, Item);
+											SequencerBaker->AddRecorderToBake(Recorder);
+										}
+									}
+								}
+							}
+						}
+					}
+				}
+			}
+		}
+	}
+}
+
+void FCommonAnimationTrackEditor::IsolateLinkedAnimTracks(TArray<FGuid> Bindings, bool bIsolate)
+{
+
+	TSharedPtr<ISequencer> SequencerPtr = GetSequencer();
+	if (!SequencerPtr.IsValid())
+	{
+		return;
+	}
+	if (bIsolate)
+	{
+		if (USequencerBakeSubsystem* SequencerBakeSubsystem = GEditor->GetEditorSubsystem<USequencerBakeSubsystem>())
+		{
+			if (TSharedPtr<UE::Sequencer::FSequencerBaker> SequencerBaker = SequencerBakeSubsystem->GetSequencerBaker(SequencerPtr))
+			{
+				SequencerBaker->BakeFullRange();
+			}
+		}
+	}
+	FScopedTransaction MatchSection(LOCTEXT("IsolateLinkedAnimTracks", "Isolate Linked Anim Tracks"));
+
+	for (FGuid& Binding : Bindings)
+	{
+		FCommonAnimationTrackEditor::IsolateLinkedAnimTrack(SequencerPtr, Binding, bIsolate);
+	}
+}
+
+bool FCommonAnimationTrackEditor::HasNoLinkedTracks(TArray<FGuid> Bindings)
+{
+	TSharedPtr<ISequencer> SequencerPtr = GetSequencer();
+	if (!SequencerPtr.IsValid())
+	{
+		return true;
+	}
+
+	for (FGuid& Binding : Bindings)
+	{
+		if (GetLinkedAnimTrack(Binding, false))
+		{
+			return false;
+		}
+	}
+	return true;
+}
+
+ECheckBoxState FCommonAnimationTrackEditor::IsLinkedAnimTrackIsolated(FGuid ObjectGuid) const
+{
+	TSharedPtr<ISequencer> SequencerPtr = GetSequencer();
+	if (!SequencerPtr.IsValid())
+	{
+		return ECheckBoxState::Unchecked;
+	}
+	if (IMovieSceneLinkedAnimTrackProvider* Provider = FindProviderForBinding(SequencerPtr, ObjectGuid))
+	{
+		if (UMovieSceneTrack* Track = Provider->GetLinkedAnimTrack(SequencerPtr, ObjectGuid))
+		{
+			return IsObjectEnabled(Track) ? ECheckBoxState::Checked : ECheckBoxState::Unchecked;
+
+		}
+	}
+	return ECheckBoxState::Unchecked;
+}
+
+void FCommonAnimationTrackEditor::OnIsolateAnimTrackClicked(ECheckBoxState NewState, FGuid ObjectGuid)
+{
+	TSharedPtr<ISequencer> SequencerPtr = GetSequencer();
+	if (!SequencerPtr.IsValid())
+	{
+		return;
+	}
+	const bool bIsolate = (NewState == ECheckBoxState::Checked);
+	TArray<FGuid> SelectedObjects;
+
+	for (TViewModelPtr<FObjectBindingModel> ObjectBindingNode : SequencerPtr->GetViewModel()->GetSelection()->Outliner.Filter<FObjectBindingModel>())
+	{
+		FGuid Guid = ObjectBindingNode->GetObjectGuid();
+		SelectedObjects.Add(Guid);
+	}
+	if (SelectedObjects.Contains(ObjectGuid) == false)
+	{
+		SelectedObjects.Add(ObjectGuid);
+	}
+	IsolateLinkedAnimTracks(SelectedObjects, bIsolate);
+}
+
+void FCommonAnimationTrackEditor::BuildObjectBindingColumnWidgets(
+	TFunctionRef<TSharedRef<SHorizontalBox>()> GetEditBox,
+	const UE::Sequencer::TViewModelPtr<UE::Sequencer::FObjectBindingModel>& ObjectBinding,
+	const UE::Sequencer::FCreateOutlinerViewParams& InParams,
+	const FName& InColumnName)
+{
+	using namespace UE::Sequencer;
+
+	FGuid ObjectGuid = ObjectBinding->GetObjectGuid();
+	if (HasNoLinkedTracks({ObjectGuid}))
+	{
+		return;
+	}
+
+	bool bIsolateAnim = false;
+	if (InColumnName == FCommonOutlinerNames::Nav)
+	{
+		bIsolateAnim = true;
+	}
+	else if (InColumnName == FCommonOutlinerNames::KeyFrame)
+	{
+		bIsolateAnim = InParams.TreeViewRow->IsColumnVisible(FCommonOutlinerNames::Nav) == false;
+	}
+	else if (InColumnName == FCommonOutlinerNames::Edit)
+	{
+		bIsolateAnim = InParams.TreeViewRow->IsColumnVisible(FCommonOutlinerNames::Nav) == false &&
+			InParams.TreeViewRow->IsColumnVisible(FCommonOutlinerNames::KeyFrame) == false;
+	}
+
+	if (bIsolateAnim)
+	{
+		const bool bEditColumn = InColumnName == FCommonOutlinerNames::Edit;
+		GetEditBox()->AddSlot()
+		.VAlign(VAlign_Center)
+		.HAlign(bEditColumn ? HAlign_Left : HAlign_Center)
+		.Padding(bEditColumn ? FMargin(4.f, 0.f) : FMargin(0.f))
+		.AutoWidth()
+		[
+			SNew(SCheckBox)
+			.Style(FAppStyle::Get(), "Sequencer.Outliner.ToggleButton")
+			.Type(ESlateCheckBoxType::ToggleButton)
+			.IsFocusable(false)
+			.IsChecked(this, &FCommonAnimationTrackEditor::IsLinkedAnimTrackIsolated, ObjectGuid)
+			.OnCheckStateChanged(this, &FCommonAnimationTrackEditor::OnIsolateAnimTrackClicked, ObjectGuid)
+			.ToolTipText(LOCTEXT("IsolateAnimTrack", "Isolate the Baked Anim Track or not"))
+			[
+				SNew(SImage)
+				.Image(FAppStyle::GetBrush("Sequencer.Outliner.Plus"))
+			]
+		];
+	}
+}
+
+bool FCommonAnimationTrackEditor::CanOpenLinkedAnimSequences(TArray<FGuid> Bindings)
+{
+	for (FGuid& Binding : Bindings)
+	{
+		if (CanOpenLinkedAnimSequence(Binding))
+		{
+			return true;
+		}
+	}
+	return false;
+}
+
 bool FCommonAnimationTrackEditor::CanOpenLinkedAnimSequence(FGuid Binding)
 {
 
@@ -2014,6 +2577,22 @@ bool FCommonAnimationTrackEditor::CanOpenLinkedAnimSequence(FGuid Binding)
 
 void FCommonAnimationTrackEditor::ConstructObjectBindingTrackMenu(FMenuBuilder& MenuBuilder, TArray<FGuid> ObjectBindings)
 {
+
+	auto AddAutoBakedMenus = [&MenuBuilder,this](const TArray<FGuid>& InObjectBindings)
+	{
+		MenuBuilder.AddSeparator();
+		MenuBuilder.AddMenuEntry(
+			LOCTEXT("ConfigureAutoBake", "Configure AutoBake"),
+			LOCTEXT("ConfigureAutoBakeTooltip", "Setup linked and automatic baked Anim Sequences"),
+			FSlateIcon(),
+			FUIAction(
+				FExecuteAction::CreateRaw(this, &FCommonAnimationTrackEditor::ConfigureAutoBake, InObjectBindings),
+				FCanExecuteAction::CreateRaw(this, &FCommonAnimationTrackEditor::HasNoLinkedTracks, InObjectBindings)
+			),
+			NAME_None,
+			EUserInterfaceActionType::Button);
+	};
+
 	if(ObjectBindings.Num() > 0)
 	{
 		USkeletalMeshComponent* SkelMeshComp = AcquireSkeletalMeshFromObjectGuid(ObjectBindings[0], GetSequencer());
@@ -2043,7 +2622,12 @@ void FCommonAnimationTrackEditor::ConstructObjectBindingTrackMenu(FMenuBuilder&
 				),
 				NAME_None,
 				EUserInterfaceActionType::Button);
-
+			
+			MenuBuilder.AddSeparator();
+		
+				AddAutoBakedMenus(ObjectBindings);
+	
+			MenuBuilder.AddSeparator();
 			MenuBuilder.AddMenuEntry(
 				LOCTEXT("CreateAnimSequence", "Bake Animation Sequence"),
 				LOCTEXT("PasteCreateAnimSequenceTooltip", "Bake an Animation Sequence for this Skeletal Mesh. Note it will create it based upon the Sequencer Display Range and Display Frame Rate"),
@@ -2065,6 +2649,16 @@ void FCommonAnimationTrackEditor::ConstructObjectBindingTrackMenu(FMenuBuilder&
 
 			MenuBuilder.EndSection();
 		}
+		else
+		{
+			//maybe topmost actor that has multiple skel mesh components so still do stuff
+			TArray<USkeletalMeshComponent*> SkelMeshComponents;
+			TArray<FGuid> BindingIds;
+			if(AcquireAllBoundSkeletalMeshObjectsFromObjectGuid(ObjectBindings[0], GetSequencer(),SkelMeshComponents, BindingIds))
+			{
+				AddAutoBakedMenus(BindingIds);
+			}
+		}
 	}
 }
 
diff --git a/Engine/Source/Editor/MovieSceneTools/Private/TrackEditors/DefaultLinkedAnimTrackProvider.cpp b/Engine/Source/Editor/MovieSceneTools/Private/TrackEditors/DefaultLinkedAnimTrackProvider.cpp
new file mode 100644
--- /dev/null
+++ b/Engine/Source/Editor/MovieSceneTools/Private/TrackEditors/DefaultLinkedAnimTrackProvider.cpp
@@ -0,0 +1,331 @@
+// Copyright Epic Games, Inc. All Rights Reserved.
+
+#include "TrackEditors/DefaultLinkedAnimTrackProvider.h"
+#include "TrackEditors/CommonAnimationTrackEditor.h"
+#include "ISequencer.h"
+#include "MovieScene.h"
+#include "LevelSequence.h"
+#include "LevelSequenceAnimSequenceLink.h"
+#include "Tracks/MovieSceneSkeletalAnimationTrack.h"
+#include "Sections/MovieSceneSkeletalAnimationSection.h"
+#include "Animation/AnimSequence.h"
+#include "Baker/SequencerBakerSubsystem.h"
+#include "Baker/SequencerBaker.h"
+#include "Editor.h"
+
+#define LOCTEXT_NAMESPACE "DefaultLinkedAnimTrackProvider"
+
+using namespace UE::Sequencer;
+
+static bool DisableTrackHelper(UMovieSceneTrack* Track, bool bInEvalDisabled)
+{
+	bool bWasChanged = false;
+	if (Track->GetMaxRowIndex() == 0)
+	{
+		if (Track->IsEvalDisabled() != bInEvalDisabled)
+		{
+			bWasChanged = true;
+			Track->Modify();
+			Track->SetEvalDisabled(bInEvalDisabled);
+		}
+	}
+	else
+	{
+		for (int32 RowIndex = 0; RowIndex <= Track->GetMaxRowIndex(); ++RowIndex)
+		{
+			if (Track->IsRowEvalDisabled(RowIndex) != bInEvalDisabled)
+			{
+				bWasChanged = true;
+				Track->Modify();
+				Track->SetRowEvalDisabled(bInEvalDisabled, RowIndex);
+			}
+		}
+	}
+	return bWasChanged;
+}
+
+bool FDefaultLinkedAnimTrackProvider::CanHandleBinding(const TWeakPtr<ISequencer>& InSequencer, FGuid BindingID) const
+{
+	// Default provider is the fallback — always handles
+	return true;
+}
+
+UMovieSceneTrack* FDefaultLinkedAnimTrackProvider::GetLinkedAnimTrack(const TWeakPtr<ISequencer>& InSequencer, FGuid BindingID) const
+{
+	TSharedPtr<ISequencer> SequencerPtr = InSequencer.Pin();
+	if (!SequencerPtr.IsValid() || SequencerPtr->GetFocusedMovieSceneSequence() == nullptr ||
+		SequencerPtr->GetFocusedMovieSceneSequence()->GetMovieScene() == nullptr)
+	{
+		return nullptr;
+	}
+	TArray<FLevelSequenceAnimSequenceLinkItem*> LinkedItems = FCommonAnimationTrackEditor::GetLinkedAnimSequences(InSequencer, BindingID);
+	if (LinkedItems.Num() > 0)
+	{
+		UMovieScene* MovieScene = SequencerPtr->GetFocusedMovieSceneSequence()->GetMovieScene();
+		TArray<UMovieSceneTrack*> AllTracks = MovieScene->FindTracks(UMovieSceneSkeletalAnimationTrack::StaticClass(), BindingID, NAME_None);
+		for (UMovieSceneTrack* Track : AllTracks)
+		{
+			if (UMovieSceneSkeletalAnimationTrack* AnimTrack = Cast<UMovieSceneSkeletalAnimationTrack>(Track))
+			{
+				bool bItMatches = false;
+				const TArray<UMovieSceneSection*>& Sections = AnimTrack->GetAllSections();
+				for (int32 Index = 0; Index < LinkedItems.Num(); ++Index)
+				{
+					FLevelSequenceAnimSequenceLinkItem* Item = LinkedItems[Index];
+					if (Item && Item->bAutoBake == true)
+					{
+						UAnimSequence* LinkedAnimSequence = Item->ResolveAnimSequence();
+						for (int32 SectionIndex = 0; SectionIndex < Sections.Num(); ++SectionIndex)
+						{
+							if (UMovieSceneSkeletalAnimationSection* AnimSection = Cast<UMovieSceneSkeletalAnimationSection>(Sections[SectionIndex]))
+							{
+								if (LinkedAnimSequence && LinkedAnimSequence == AnimSection->Params.Animation)
+								{
+									bItMatches = true;
+									break;
+								}
+							}
+						}
+					}
+				}
+				if (bItMatches)
+				{
+					return AnimTrack;
+				}
+			}
+		}
+	}
+	return nullptr;
+}
+
+UMovieSceneTrack* FDefaultLinkedAnimTrackProvider::GetOrCreateLinkedAnimTrack(
+	const TWeakPtr<ISequencer>& InSequencer,
+	FGuid BindingID,
+	FCommonAnimationTrackEditor& TrackEditor)
+{
+	// First try read-only lookup
+	if (UMovieSceneTrack* Track = GetLinkedAnimTrack(InSequencer, BindingID))
+	{
+		return Track;
+	}
+
+	TSharedPtr<ISequencer> SequencerPtr = InSequencer.Pin();
+	if (!SequencerPtr.IsValid() || SequencerPtr->GetFocusedMovieSceneSequence() == nullptr ||
+		SequencerPtr->GetFocusedMovieSceneSequence()->GetMovieScene() == nullptr)
+	{
+		return nullptr;
+	}
+
+	UMovieScene* MovieScene = SequencerPtr->GetFocusedMovieSceneSequence()->GetMovieScene();
+	FScopedTransaction MatchSection(LOCTEXT("CreateLinkedAnimTrack_Transaction", "Create Linked Anim Track"));
+	MovieScene->Modify();
+	FMovieSceneBinding* Binding = MovieScene->FindBinding(BindingID);
+	if (!Binding)
+	{
+		return nullptr;
+	}
+
+	UMovieSceneSkeletalAnimationTrack* AnimTrack = TrackEditor.FindEmptyAnimationTrack(BindingID);
+	if (AnimTrack == nullptr)
+	{
+		AnimTrack = Cast<UMovieSceneSkeletalAnimationTrack>(TrackEditor.AddTrack(MovieScene, BindingID, UMovieSceneSkeletalAnimationTrack::StaticClass(), NAME_None));
+		
+	}
+	if (!AnimTrack)
+	{
+		return nullptr;
+	}
+
+	const FString TrackName = TEXT("Animation (Linked}");
+	AnimTrack->SetDisplayName(FText::FromString(TrackName));
+	AnimTrack->SetColorTint(UMovieSceneSkeletalAnimationTrack::LinkedAnimSeqTrackColor);
+
+	AnimTrack->Modify();
+	USkeletalMeshComponent* SkelMeshComp = FCommonAnimationTrackEditor::AcquireSkeletalMeshFromObjectGuid(BindingID, SequencerPtr);
+	USkeleton* Skeleton = FCommonAnimationTrackEditor::AcquireSkeletonFromObjectGuid(BindingID, SequencerPtr);
+	bool bNeedToCreateLinkedAnim = true;
+	int32 LinkedItemNum = INDEX_NONE;
+	TArray<FLevelSequenceAnimSequenceLinkItem*> LinkedItems = FCommonAnimationTrackEditor::GetLinkedAnimSequences(SequencerPtr, BindingID);
+	if (LinkedItems.Num() > 0)
+	{
+		for (int32 LinkedIndex = 0; LinkedIndex < LinkedItems.Num(); ++LinkedIndex)
+		{
+			FLevelSequenceAnimSequenceLinkItem* LinkedItem = LinkedItems[LinkedIndex];
+			if (LinkedItem && LinkedItem->bUseCustomTimeRange == false)
+			{
+				bNeedToCreateLinkedAnim = false;
+				LinkedItemNum = LinkedIndex;
+			}
+		}
+	}
+	if (bNeedToCreateLinkedAnim)
+	{
+		TrackEditor.CreateLinkedAnimationSequence(SkelMeshComp, Skeleton, BindingID, bNeedToCreateLinkedAnim);
+		LinkedItems = FCommonAnimationTrackEditor::GetLinkedAnimSequences(SequencerPtr, BindingID);
+		LinkedItemNum = LinkedItems.Num() - 1;
+	}
+	if (LinkedItems.Num() > 0 && LinkedItemNum != INDEX_NONE)
+	{
+		if (FLevelSequenceAnimSequenceLinkItem* LinkedItem = LinkedItems[LinkedItemNum])
+		{
+			UAnimSequence* AnimSequence = LinkedItem->ResolveAnimSequence();
+			TOptional<TRange<FFrameNumber>> OptionalRange = SequencerPtr->GetSubSequenceRange();
+			TRange<FFrameNumber> Range = OptionalRange.IsSet() ? OptionalRange.GetValue() : SequencerPtr->GetFocusedMovieSceneSequence()->GetMovieScene()->GetPlaybackRange();
+			FFrameNumber Frame = Range.GetLowerBoundValue();
+			UMovieSceneSkeletalAnimationSection* NewSection = Cast<UMovieSceneSkeletalAnimationSection>(AnimTrack->AddNewAnimationOnRow(Frame, AnimSequence, LinkedItems.Num() - 1));
+			if (NewSection)
+			{
+				LinkedItem->bUseCustomTimeRange = false;
+				LinkedItem->bAutoBake = true;
+				NewSection->SetRange(Range);
+				AnimTrack->SetEvalDisabled(true);
+			}
+			return AnimTrack;
+		}
+	}
+
+	return nullptr;
+}
+
+bool FDefaultLinkedAnimTrackProvider::IsolateLinkedAnimTrack(const TWeakPtr<ISequencer>& InSequencer, FGuid BindingID, bool bIsolate)
+{
+	bool bWasChanged = false;
+
+	TSharedPtr<ISequencer> SequencerPtr = InSequencer.Pin();
+	if (!SequencerPtr.IsValid() || SequencerPtr->GetFocusedMovieSceneSequence() == nullptr ||
+		SequencerPtr->GetFocusedMovieSceneSequence()->GetMovieScene() == nullptr)
+	{
+		return bWasChanged;
+	}
+	TArray<FLevelSequenceAnimSequenceLinkItem*> LinkedItems = FCommonAnimationTrackEditor::GetLinkedAnimSequences(InSequencer, BindingID);
+	if (LinkedItems.Num() > 0)
+	{
+		UMovieScene* MovieScene = SequencerPtr->GetFocusedMovieSceneSequence()->GetMovieScene();
+		TArray<UMovieSceneTrack*> AllTracks = MovieScene->FindTracks(UMovieSceneTrack::StaticClass(), BindingID, NAME_None);
+		for (UMovieSceneTrack* Track : AllTracks)
+		{
+			if (UMovieSceneSkeletalAnimationTrack* AnimTrack = Cast<UMovieSceneSkeletalAnimationTrack>(Track))
+			{
+				const TArray<UMovieSceneSection*>& Sections = AnimTrack->GetAllSections();
+				bool bItMatches = false;
+				for (int32 Index = 0; Index < LinkedItems.Num(); ++Index)
+				{
+					FLevelSequenceAnimSequenceLinkItem* Item = LinkedItems[Index];
+					if (Item && Item->bAutoBake == true)
+					{
+						for (int32 SectionIndex = 0; SectionIndex < Sections.Num(); ++SectionIndex)
+						{
+							if (UMovieSceneSkeletalAnimationSection* AnimSection = Cast<UMovieSceneSkeletalAnimationSection>(Sections[SectionIndex]))
+							{
+								UAnimSequence* LinkedAnimSequence = Item->ResolveAnimSequence();
+								if (LinkedAnimSequence && LinkedAnimSequence == AnimSection->Params.Animation)
+								{
+									bItMatches = true;
+									break;
+								}
+							}
+						}
+					}
+				}
+				if (bItMatches)
+				{
+					if (DisableTrackHelper(AnimTrack, !bIsolate))
+					{
+						bWasChanged = true;
+					}
+				}
+				else
+				{
+					if (DisableTrackHelper(AnimTrack, bIsolate))
+					{
+						bWasChanged = true;
+					}
+				}
+			}
+			else
+			{
+				FString TrackName = Track->GetClass()->GetName();
+				if (TrackName == FString("MovieSceneControlRigParameterTrack"))
+				{
+					if (DisableTrackHelper(Track, bIsolate))
+					{
+						bWasChanged = true;
+					}
+				}
+			}
+		}
+		if (USequencerBakeSubsystem* SequencerBakeSubsystem = GEditor->GetEditorSubsystem<USequencerBakeSubsystem>())
+		{
+			if (TSharedPtr<UE::Sequencer::FSequencerBaker> SequencerBaker = SequencerBakeSubsystem->GetSequencerBaker(SequencerPtr))
+			{
+				SequencerBaker->IgnoreLastChange();
+			}
+		}
+	}
+	return bWasChanged;
+}
+
+void FDefaultLinkedAnimTrackProvider::UpdateLinkedAnimSectionRange(const TWeakPtr<ISequencer>& InSequencer, FGuid BindingID)
+{
+	TSharedPtr<ISequencer> SequencerPtr = InSequencer.Pin();
+	if (!SequencerPtr.IsValid())
+	{
+		return;
+	}
+	ULevelSequence* LevelSequence = Cast<ULevelSequence>(SequencerPtr->GetFocusedMovieSceneSequence());
+	if (LevelSequence && LevelSequence->GetClass()->ImplementsInterface(UInterface_AssetUserData::StaticClass()))
+	{
+		USequencerBakeSubsystem* SequencerBakeSubsystem = GEditor->GetEditorSubsystem<USequencerBakeSubsystem>();
+		TSharedPtr<UE::Sequencer::FSequencerBaker> SequencerBaker = SequencerBakeSubsystem ? SequencerBakeSubsystem->GetSequencerBaker(SequencerPtr) : nullptr;
+		if (SequencerBaker.IsValid() == false)
+		{
+			return;
+		}
+
+		if (IInterface_AssetUserData* AssetUserDataInterface = Cast< IInterface_AssetUserData >(LevelSequence))
+		{
+			ULevelSequenceAnimSequenceLink* LevelAnimLink = AssetUserDataInterface->GetAssetUserData< ULevelSequenceAnimSequenceLink >();
+
+			UMovieScene* MovieScene = SequencerPtr->GetFocusedMovieSceneSequence()->GetMovieScene();
+			TOptional<TRange<FFrameNumber>> OptionalRange = SequencerPtr->GetSubSequenceRange();
+			TRange<FFrameNumber> Range = OptionalRange.IsSet() ? OptionalRange.GetValue() : SequencerPtr->GetFocusedMovieSceneSequence()->GetMovieScene()->GetPlaybackRange();
+
+			TArray<FLevelSequenceAnimSequenceLinkItem*> LinkedItems = FCommonAnimationTrackEditor::GetLinkedAnimSequences(InSequencer, BindingID);
+			if (LinkedItems.Num() > 0)
+			{
+				TArray<UMovieSceneTrack*> AllTracks = MovieScene->FindTracks(UMovieSceneTrack::StaticClass(), BindingID, NAME_None);
+				for (UMovieSceneTrack* Track : AllTracks)
+				{
+					if (UMovieSceneSkeletalAnimationTrack* AnimTrack = Cast<UMovieSceneSkeletalAnimationTrack>(Track))
+					{
+						const TArray<UMovieSceneSection*>& Sections = AnimTrack->GetAllSections();
+						for (int32 Index = 0; Index < LinkedItems.Num(); ++Index)
+						{
+							FLevelSequenceAnimSequenceLinkItem* Item = LinkedItems[Index];
+							if (Item && Item->bAutoBake == true &&  Item->bUseCustomTimeRange == false)
+							{
+								for (int32 SectionIndex = 0; SectionIndex < Sections.Num(); ++SectionIndex)
+								{
+									if (UMovieSceneSkeletalAnimationSection* AnimSection = Cast<UMovieSceneSkeletalAnimationSection>(Sections[SectionIndex]))
+									{
+										UAnimSequence* LinkedAnimSequence = Item->ResolveAnimSequence();
+										if (LinkedAnimSequence && LinkedAnimSequence == AnimSection->Params.Animation)
+										{
+											if (AnimSection->GetRange() != Range)
+											{
+												AnimSection->Modify();
+												AnimSection->SetRange(Range);
+											}
+										}
+									}
+								}
+							}
+						}
+					}
+				}
+			}
+		}
+	}
+}
+
+#undef LOCTEXT_NAMESPACE
diff --git a/Engine/Source/Editor/MovieSceneTools/Private/TrackEditors/DefaultLinkedAnimTrackProvider.h b/Engine/Source/Editor/MovieSceneTools/Private/TrackEditors/DefaultLinkedAnimTrackProvider.h
new file mode 100644
--- /dev/null
+++ b/Engine/Source/Editor/MovieSceneTools/Private/TrackEditors/DefaultLinkedAnimTrackProvider.h
@@ -0,0 +1,23 @@
+// Copyright Epic Games, Inc. All Rights Reserved.
+
+#pragma once
+
+#include "IMovieSceneLinkedAnimTrackProvider.h"
+
+/**
+ * Default implementation of IMovieSceneLinkedAnimTrackProvider.
+ * Handles the standard case where the linked animation track is a standalone
+ * UMovieSceneSkeletalAnimationTrack on the binding (not inside a mixer).
+ */
+class FDefaultLinkedAnimTrackProvider : public IMovieSceneLinkedAnimTrackProvider
+{
+public:
+
+	//~ IMovieSceneLinkedAnimTrackProvider interface
+	virtual int32 GetLinkedAnimTrackPriority() const override { return 100; }
+	virtual bool CanHandleBinding(const TWeakPtr<ISequencer>& InSequencer, FGuid BindingID) const override;
+	virtual UMovieSceneTrack* GetLinkedAnimTrack(const TWeakPtr<ISequencer>& InSequencer, FGuid BindingID) const override;
+	virtual UMovieSceneTrack* GetOrCreateLinkedAnimTrack(const TWeakPtr<ISequencer>& InSequencer, FGuid BindingID, UE::Sequencer::FCommonAnimationTrackEditor& TrackEditor) override;
+	virtual bool IsolateLinkedAnimTrack(const TWeakPtr<ISequencer>& InSequencer, FGuid BindingID, bool bIsolate) override;
+	virtual void UpdateLinkedAnimSectionRange(const TWeakPtr<ISequencer>& InSequencer, FGuid BindingID) override;
+};
diff --git a/Engine/Source/Editor/MovieSceneTools/Private/TrackEditors/SkeletalAnimationTrackEditor.cpp b/Engine/Source/Editor/MovieSceneTools/Private/TrackEditors/SkeletalAnimationTrackEditor.cpp
--- a/Engine/Source/Editor/MovieSceneTools/Private/TrackEditors/SkeletalAnimationTrackEditor.cpp
+++ b/Engine/Source/Editor/MovieSceneTools/Private/TrackEditors/SkeletalAnimationTrackEditor.cpp
@@ -70,6 +70,33 @@ void FSkeletalAnimationTrackEditor::PopulateAddMixerItemMenu(FMenuBuilder& MenuB
 	}
 }
 
+void FSkeletalAnimationTrackEditor::BuildObjectBindingContextMenu(FMenuBuilder& MenuBuilder, const TArray<FGuid>& ObjectBindings, const UClass* ObjectClass)
+{
+	// Check if the Animation Mixer plugin is enabled, only do this if not, otherwise we get duplicates
+	const bool bAnimationMixerEnabled = FModuleManager::Get().IsModuleLoaded(TEXT("MovieSceneAnimMixerEditor"));
+
+	if (bAnimationMixerEnabled == false)
+	{
+		FCommonAnimationTrackEditor::BuildObjectBindingContextMenu(MenuBuilder, ObjectBindings, ObjectClass);
+	}
+}
+
+
+void FSkeletalAnimationTrackEditor::BuildObjectBindingColumnWidgets(
+	TFunctionRef<TSharedRef<SHorizontalBox>()> GetEditBox,
+	const UE::Sequencer::TViewModelPtr<UE::Sequencer::FObjectBindingModel>& ObjectBinding,
+	const UE::Sequencer::FCreateOutlinerViewParams& InParams,
+	const FName& InColumnName)
+{
+	// Check if the Animation Mixer plugin is enabled, only do this if not, otherwise we get duplicates
+	const bool bAnimationMixerEnabled = FModuleManager::Get().IsModuleLoaded(TEXT("MovieSceneAnimMixerEditor"));
+
+	if (bAnimationMixerEnabled == false)
+	{
+		FCommonAnimationTrackEditor::BuildObjectBindingColumnWidgets(GetEditBox, ObjectBinding, InParams, InColumnName);
+	}
+}
+
 void FSkeletalAnimationTrackEditor::PopulateObjectBindingAnimationMenu(FMenuBuilder& MenuBuilder, const TArray<FGuid>& ObjectBindings, const UClass* ObjectClass, bool bIsInsideSubmenu)
 {
 	if (!ObjectClass || (!ObjectClass->IsChildOf(USkeletalMeshComponent::StaticClass()) && !ObjectClass->IsChildOf(AActor::StaticClass()) && !ObjectClass->IsChildOf(UChildActorComponent::StaticClass())))
diff --git a/Engine/Source/Editor/MovieSceneTools/Private/TrackEditors/SkeletalAnimationTrackEditor.h b/Engine/Source/Editor/MovieSceneTools/Private/TrackEditors/SkeletalAnimationTrackEditor.h
--- a/Engine/Source/Editor/MovieSceneTools/Private/TrackEditors/SkeletalAnimationTrackEditor.h
+++ b/Engine/Source/Editor/MovieSceneTools/Private/TrackEditors/SkeletalAnimationTrackEditor.h
@@ -39,7 +39,12 @@ public:
 	virtual int32 GetMixerItemMenuPriority() const override { return 0; } // Highest priority in mixer '+' menu
 	virtual void PopulateObjectBindingAnimationMenu(FMenuBuilder& MenuBuilder, const TArray<FGuid>& ObjectBindings, const UClass* ObjectClass, bool bIsInsideSubmenu) override;
 	virtual int32 GetObjectBindingAnimationMenuPriority() const override { return 25; } // After Animation Mixer in Object Binding menu
-
+	virtual void BuildObjectBindingContextMenu(FMenuBuilder& MenuBuilder, const TArray<FGuid>& ObjectBindings, const UClass* ObjectClass) override;
+	virtual void BuildObjectBindingColumnWidgets(
+		TFunctionRef<TSharedRef<SHorizontalBox>()> GetEditBox,
+		const UE::Sequencer::TViewModelPtr<UE::Sequencer::FObjectBindingModel>& ObjectBinding,
+		const UE::Sequencer::FCreateOutlinerViewParams& InParams,
+		const FName& InColumnName) override;
 private:
 
 	void OnPostPropertyChanged(UObject* InObject, struct FPropertyChangedEvent& InPropertyChangedEvent);
diff --git a/Engine/Source/Editor/MovieSceneTools/Public/Baker/ISequencerBaker.h b/Engine/Source/Editor/MovieSceneTools/Public/Baker/ISequencerBaker.h
new file mode 100644
--- /dev/null
+++ b/Engine/Source/Editor/MovieSceneTools/Public/Baker/ISequencerBaker.h
@@ -0,0 +1,118 @@
+// Copyright Epic Games, Inc. All Rights Reserved.
+
+#pragma once
+
+#include "CoreMinimal.h"
+#include "Misc/FrameRate.h"
+#include "Math/Range.h"
+#include "EditorSubsystem.h"
+#include "MovieSceneTimeUnit.h"
+#include "UObject/Interface.h"
+#include "ISequencerBaker.generated.h"
+
+#define UE_API MOVIESCENETOOLS_API
+
+namespace UE::AIE
+{
+	/**
+	* FBakeTimeRange
+	*
+	* Object that contains a range of frame numbers that specify an interval to bake over
+	*/
+	struct FBakeTimeRange
+	{
+		FBakeTimeRange() :
+			StartFrame(0), EndFrame(0), DisplayRate(1, 1), TickResolution(1, 1)
+		{
+			IncomingRange = TRange<FFrameNumber>{ 0 };
+		};
+		~FBakeTimeRange() = default;
+		FFrameNumber StartFrame;
+		FFrameNumber EndFrame;
+		FFrameNumber FrameStep;
+		FFrameRate DisplayRate;
+		FFrameRate TickResolution;
+		int32 NumFrames = 0;
+		TRange<FFrameNumber> IncomingRange;
+
+		UE_API bool UpdateIfNeeded(const TRange<FFrameNumber>& InRange, const FFrameRate& InDisplayRate, const FFrameRate& InTickResolution);
+		UE_API int32 CalculateIndex(const FFrameNumber& InCurrentFrame) const;
+	};
+
+	/**
+	* FBakeTimeIndex
+	*
+	* A specific index in the FBakeTimeRange, used to eaily keep track of an index and how it relates to a frame number or a second.
+	* This is passed into the recorders on each frame baked.
+	*/
+	struct FBakeTimeIndex
+	{
+		int32 Index;
+		FFrameNumber FrameNumber;
+		double DeltaTime;
+	};
+}
+
+UINTERFACE(MinimalAPI)
+class USequencerBakeRecorder : public UInterface
+{
+	GENERATED_BODY()
+};
+
+/**
+* ISequencerBakeRecorder
+*
+* Main interface for an object that records a specific set of data when the bake happens.  For eample an AnimationSequenceBaker would record a specific
+* animation sequence for a specific skel mesh component that you want to record.
+*/
+class ISequencerBakeRecorder
+{
+	GENERATED_BODY()
+
+public:
+	//virtual ~ISequencerBakeRecorder() is defined above
+	//unique ID
+	virtual uint32 GetHash() = 0;
+	virtual void BakeStarted(const UE::AIE::FBakeTimeRange& InRange) = 0;
+	virtual void BakeFrame(const UE::AIE::FBakeTimeIndex& InIndex) = 0;
+	virtual void BakeCancelled() = 0;
+	virtual void BakeFinished() = 0;
+	virtual bool HasSequencerBinding(FGuid InGuid) = 0;
+	virtual bool IsolateBakeResult(bool bIsolate) = 0;
+	virtual void SetRecordingEnabled(bool bInEnabled) = 0;
+	virtual bool GetRecordingEnabled() = 0;
+	virtual void ReadyToBake() = 0;
+	virtual void RemovedFromBake() = 0;
+};
+
+namespace UE::AIE
+{
+
+/**
+* ISequencerBaker
+*
+* Main interface for the object performs a bake on a set of ISequencerBakeRecorders
+*/
+
+class ISequencerBaker 
+{
+public:
+	struct FRecorderOptions
+	{
+		FFrameNumber WarmupFrames = FFrameNumber(0);
+		FFrameNumber DelayBeforeStart = FFrameNumber(0);
+	};
+	virtual ~ISequencerBaker() {};
+	virtual void AddRecorder(TSharedPtr<ISequencerBakeRecorder>& InRecorder, const FRecorderOptions& InOptions) = 0;
+	virtual void RemoveRecorder(TSharedPtr<ISequencerBakeRecorder>& InRecorder) = 0;
+	virtual FBakeTimeRange GetTimeRange() = 0;
+
+	virtual bool IsBakeRunning(TOptional<float>& OutPercentageDone) = 0;
+	virtual void CancelBake() = 0;
+
+};
+
+};
+
+#undef UE_API
+
diff --git a/Engine/Source/Editor/MovieSceneTools/Public/Baker/SequencerBaker.h b/Engine/Source/Editor/MovieSceneTools/Public/Baker/SequencerBaker.h
new file mode 100644
--- /dev/null
+++ b/Engine/Source/Editor/MovieSceneTools/Public/Baker/SequencerBaker.h
@@ -0,0 +1,215 @@
+// Copyright Epic Games, Inc. All Rights Reserved.
+
+#pragma once
+
+#include "Delegates/Delegate.h"
+#include "Evaluation/MovieSceneSequenceTransform.h"
+#include "ISequencerBaker.h"
+#include "MovieSceneSequenceID.h"
+#include "MovieSceneSignedObject.h"
+#include "TickableEditorObject.h"
+
+class FSequencerBakeRecorder;
+class ISequencer;
+class UMovieSceneTrack;
+struct FGuid;
+class USkeletalMeshComponent;
+class IMovieSceneToolsAnimationBakeHelper;
+class UWorld;
+
+#define UE_API MOVIESCENETOOLS_API
+
+namespace UE::Sequencer
+{
+	//delegates fired when we add or remove a bake recorder
+	DECLARE_MULTICAST_DELEGATE_OneParam(FOnAddBakeRecorder, ISequencerBakeRecorder*);
+	DECLARE_MULTICAST_DELEGATE_OneParam(FOnRemoveBakeRecorder, ISequencerBakeRecorder*);
+
+	/**
+	* FSequencerBaker
+	*
+	* Main object that Bakes within ISequencer, using Sequencer evaluation over the range.
+	* It keeps track of track signatures for each recorder to seek if they need to get baked when we switch over to using the baked track
+	*/
+	class FSequencerBaker
+		: public UE::AIE::ISequencerBaker
+		, public FTickableEditorObject
+		, public TSharedFromThis<FSequencerBaker>
+	{
+
+	public:
+		FSequencerBaker();
+		virtual ~FSequencerBaker() {};
+
+		void Initialize(TSharedPtr<ISequencer>& InSequencer);
+		void Release();
+
+		//~ ISequencerBaker
+		virtual UE::AIE::FBakeTimeRange GetTimeRange() override;
+		virtual bool IsBakeRunning(TOptional<float>& OutPercentageDone) override;
+		virtual void CancelBake() override;
+		virtual  void AddRecorder(TSharedPtr<ISequencerBakeRecorder>& InRecorder, const UE::AIE::ISequencerBaker::FRecorderOptions& InOptions) override;
+		virtual  void RemoveRecorder(TSharedPtr<ISequencerBakeRecorder>& InRecorder) override;
+
+		//~ FTickableEditorObject Interface
+		virtual void Tick(float DeltaTime) override;
+		virtual ETickableTickType GetTickableTickType() const override { return ETickableTickType::Always; }
+		virtual TStatId GetStatId() const override { RETURN_QUICK_DECLARE_CYCLE_STAT(FSequencerBaker, STATGROUP_Tickables); }
+
+	public:
+	
+		//add all existing recorders to get baked
+		UE_API void AddAllRecordersToBake();
+		//add extra recorders to get baked just once on next bake
+		UE_API void AddExtraRecordersToBake(const TArray<TSharedPtr<ISequencerBakeRecorder>>& InExtraRecorders);
+		//bake everything and will remove recorders from bake queue
+		UE_API void BakeFullRange();
+	
+		//this adds that recorder that was already added to be baked on next bake time
+		UE_API bool AddRecorderToBake(TSharedPtr<ISequencerBakeRecorder>& InRecorder);
+
+		UE_API void IsolateRecorderWithGuid(FGuid InGuid, bool bIsolate);
+
+		UE_API void AddSkelMeshCompsToTick(TArray<USkeletalMeshComponent*>& InSkelMeshCompArray);
+
+		UE_API void AddSignedObjectsToTrack(TArray<TWeakObjectPtr<UMovieSceneSignedObject>>& InSignedObjects, TSharedPtr<ISequencerBakeRecorder>& Owner);
+		UE_API void RemoveSignedObjectsFromTrack(TArray<TWeakObjectPtr<UMovieSceneSignedObject>>& InSignedObjects, TSharedPtr<ISequencerBakeRecorder>& Owner);
+	
+		//ignore the last change, used mostly for when we isolate the tracks we don't want to trigger a re-baking
+		UE_API void IgnoreLastChange();
+		
+		TWeakPtr<ISequencer> GetWeakSequencer() const { return WeakSequencer; }
+
+		//finds the recorder by it's signature of arguments
+		template<typename... Args>
+		TSharedPtr<ISequencerBakeRecorder>* FindRecorder(Args&&... InArgs) {
+
+			uint32 HashValue = (GetTypeHash(InArgs) + ...);
+
+			return Recorders.Find(HashValue);
+		}
+
+	public:
+
+		//public delegates folks can subscribe for when recorders are added or removed
+		FOnAddBakeRecorder OnAddBakeRecorder;
+		FOnRemoveBakeRecorder OnRemoveBakeRecorder;
+	private:
+
+		//delegates we subscribe
+		void OnPreNavigateToSequenceChanged();
+		void OnActivateSequenceChanged(FMovieSceneSequenceIDRef ID);
+		void OnMovieSceneBindingsChanged();
+		FDelegateHandle OnPreNavigateToSequenceChangedHandle;
+		FDelegateHandle OnActivateSequenceChangedHandle;
+		FDelegateHandle OnMovieSceneBindingsChangedHandle;
+
+		//called when we activate a new sequence
+		void SetUpBakeRecorders();
+
+	private:
+
+		//internal struct to set up time range based upon sequencer frame range
+		struct FBakeInterval
+		{
+			FBakeInterval() = default;
+
+			bool SetFromSequencer(ISequencer* InSequencer);
+			void Clear();
+			bool IsFrameCalcuated(int32 Index);
+			void SetFrameCalculated(int32 Index);
+
+			//will return INDEX_NONE if no more to calculate
+			int32 NextFrameToCalculate();
+			bool IsDoneCalculating() const;
+
+			UE::AIE::FBakeTimeRange TimeRange;
+			FFrameNumber DelayBeforeStart = FFrameNumber(0);
+			FFrameNumber WarmupFrames = FFrameNumber(0);
+
+		private:
+			TBitArray<> CalculatedFrames;
+
+		};
+		
+	private:
+		//internal objects that contains a movie scene signed object(track) last guid and the bakers that reference
+		//we keep of these guids and if the LastValidMovieSceneGuid changes we see which of these track guids have changed.
+		//this was we only re-bake what's needed.
+		struct FSignedObjectState
+		{
+			TSet<TSharedPtr<ISequencerBakeRecorder>> Owners;
+			FGuid Guid;
+			TArray<bool> DisabledArray; //if it's a track we keep track of this since if it get's enabled/disabled we ignore the change
+		};
+		struct FSignedObjectChanges
+		{
+			void AddObject(UMovieSceneSignedObject* InObject, TSharedPtr<ISequencerBakeRecorder>& Owner);
+			void RemoveObject(UMovieSceneSignedObject* InObject, TSharedPtr<ISequencerBakeRecorder>& Owner);
+			void AddIgnoreObject(UMovieSceneSignedObject* InObject) { if (InObject) { IgnoredObjects.Add(InObject); } };
+			void RemoveIgnoreObject(UMovieSceneSignedObject* InObject) { if (InObject) { IgnoredObjects.Remove(InObject); } };
+
+			//returns list of recorders that have changed
+			TArray<TSharedPtr<ISequencerBakeRecorder>> UpdateSignatures();
+			TMap<TWeakObjectPtr<UMovieSceneSignedObject>, FSignedObjectState>  ObjectSignatures;
+			TSet <TWeakObjectPtr<UMovieSceneSignedObject>> IgnoredObjects;
+		};
+
+		bool CheckForChanges();
+		FGuid LastValidMovieSceneGuid;
+	
+private:
+
+		//sequencer and it's interval to bake
+		TWeakPtr<ISequencer> WeakSequencer;
+		FBakeInterval  BakeInterval;
+
+		//contains full set of options e.g. if one anim sequence has delay frame others will also
+		FRecorderOptions RecorderOptions;
+
+		//Registered Recorders
+		//uses hash as key
+		TMap<uint32, TSharedPtr<ISequencerBakeRecorder>> Recorders;
+
+		//Recorders that are baking
+		TSet<TSharedPtr<ISequencerBakeRecorder>> BakingRecorders;
+		TArray<TSharedPtr<ISequencerBakeRecorder>>  LastBakedRecorders;
+
+		//skel meshes to tick that comes from the recorders
+		TSet<TWeakObjectPtr<USkeletalMeshComponent>> SkeletalMeshCompsTick;
+
+		//tracked objecs for guid changes
+		FSignedObjectChanges  SignedObjectsToTrack;
+
+	private:
+		//cached items during the bake
+		TArray<IMovieSceneToolsAnimationBakeHelper*> BakeHelpers;
+		FMovieSceneInverseSequenceTransform LocalToRootTransform;
+		//this is cached during set of functions
+		UWorld* World;
+	
+private:
+
+		//functions used during the Bake
+		bool bNeedsBaking = false;
+		bool bIsBaking = false;
+		void BakeStarted();
+		void BakeFrame();
+		void BakeCancelled();
+		void BakeFinished();
+		void BakeFramesTillDone();
+		bool TickFrameInternal(const FFrameNumber& FrameNumber, float DeltaTime);
+
+	private:
+
+		//used when we navigate to another sequence, we isolate and turn on the baked items for performance
+		void IsolateBakedTracks(TArray<TSharedPtr<ISequencerBakeRecorder>>& InRecorders, bool bIsolate);
+	};
+
+};
+
+#undef UE_API
+
+
+
+
diff --git a/Engine/Source/Editor/MovieSceneTools/Public/Baker/SequencerBakerSubsystem.h b/Engine/Source/Editor/MovieSceneTools/Public/Baker/SequencerBakerSubsystem.h
new file mode 100644
--- /dev/null
+++ b/Engine/Source/Editor/MovieSceneTools/Public/Baker/SequencerBakerSubsystem.h
@@ -0,0 +1,49 @@
+// Copyright Epic Games, Inc. All Rights Reserved.
+
+#pragma once
+
+
+#include "CoreMinimal.h"
+#include "EditorSubsystem.h"
+#include "Editor/UnrealEdTypes.h"
+#include "SequencerBakerSubsystem.generated.h"
+
+class ISequencer;
+namespace UE::Sequencer
+{
+	class FSequencerBaker;
+}
+
+#define UE_API MOVIESCENETOOLS_API
+
+/**
+* USequencerBakeSubsystem
+* 
+* Simple Editor Subsystem that acts like a Singleton, that holds a Sequencer Baker object for each open ISequencer object.
+*/
+UCLASS(MinimalAPI)
+class USequencerBakeSubsystem :
+	public UEditorSubsystem
+{
+	GENERATED_BODY()
+
+public:
+
+	UE_API virtual void Initialize(FSubsystemCollectionBase& Collection) override;
+	UE_API virtual void Deinitialize() override;
+
+	UE_API void OnSequencerCreated(TSharedRef<ISequencer> InSequencer);
+
+	UE_API void OnSequencerClosed(TSharedRef<ISequencer> InSequencer);
+
+	UE_API TSharedPtr<UE::Sequencer::FSequencerBaker> GetSequencerBaker(const TSharedPtr<ISequencer>& InSequencer);
+private:
+
+	FDelegateHandle OnSequencerCreatedHandle;
+
+	//list of created sequencers
+	TMap<TWeakPtr<ISequencer>, TSharedPtr<UE::Sequencer::FSequencerBaker>> Sequencers;
+
+};
+
+#undef UE_API
diff --git a/Engine/Source/Editor/MovieSceneTools/Public/IMovieSceneLinkedAnimTrackProvider.h b/Engine/Source/Editor/MovieSceneTools/Public/IMovieSceneLinkedAnimTrackProvider.h
new file mode 100644
--- /dev/null
+++ b/Engine/Source/Editor/MovieSceneTools/Public/IMovieSceneLinkedAnimTrackProvider.h
@@ -0,0 +1,71 @@
+// Copyright Epic Games, Inc. All Rights Reserved.
+
+#pragma once
+
+#include "CoreMinimal.h"
+#include "Features/IModularFeature.h"
+
+class UMovieSceneTrack;
+class ISequencer;
+struct FGuid;
+
+namespace UE::Sequencer { class FCommonAnimationTrackEditor; }
+
+#define UE_API MOVIESCENETOOLS_API
+/**
+ * Modular feature interface for providing linked animation track or section operations.
+ *
+ * Implementations handle finding, creating, isolating, and updating linked animation tracks
+ * that contain auto-baked animation sequences. Different providers handle different track
+ * configurations (e.g., standalone skeletal animation tracks vs. animation mixer sections/tracks).
+ *
+ * Providers are queried by priority (lower = higher priority). The first provider whose
+ * CanHandleBinding() returns true is used for that binding.
+ */
+class  IMovieSceneLinkedAnimTrackProvider : public IModularFeature
+{
+public:
+
+	UE_API static FName GetModularFeatureName();
+
+	virtual ~IMovieSceneLinkedAnimTrackProvider() {}
+
+	/** Priority for provider selection. Lower value = higher priority (checked first). Default provider uses 100. */
+	virtual int32 GetLinkedAnimTrackPriority() const = 0;
+
+	/** Returns true if this provider should handle linked anim track operations for this binding. */
+	virtual bool CanHandleBinding(const TWeakPtr<ISequencer>& InSequencer, FGuid BindingID) const = 0;
+
+	/**
+	 * Find the skeletal animation track(section) containing the linked animation sequence.
+	 * Read-only lookup, does not create anything.
+	 * @return The UMovieSceneTrack containing the linked anim, or nullptr if not found.
+	 */
+	virtual UMovieSceneTrack* GetLinkedAnimTrack(const TWeakPtr<ISequencer>& InSequencer, FGuid BindingID) const = 0;
+
+	/**
+	 * Find or create the linked animation track setup.
+	 * If no linked track exists and creation is needed, creates the necessary track structure,
+	 * linked animation sequence, section, and disables evaluation on the containing track.
+	 * @param TrackEditor The track editor instance, needed for track creation operations.
+	 * @return The UMovieSceneTrack containing the linked anim, or nullptr on failure.
+	 */
+	virtual UMovieSceneTrack* GetOrCreateLinkedAnimTrack(
+		const TWeakPtr<ISequencer>& InSequencer,
+		FGuid BindingID,
+		UE::Sequencer::FCommonAnimationTrackEditor& TrackEditor) = 0;
+
+	/**
+	 * Isolate or un-isolate the linked animation track.
+	 * When isolating: enables the linked anim track and disables other animation/control rig tracks.
+	 * When un-isolating: restores the previous evaluation state.
+	 * @return true if any track evaluation state was changed.
+	 */
+	virtual bool IsolateLinkedAnimTrack(const TWeakPtr<ISequencer>& InSequencer, FGuid BindingID, bool bIsolate) = 0;
+
+	/** Update linked anim section ranges after bake, in case the playback range changed. */
+	virtual void UpdateLinkedAnimSectionRange(const TWeakPtr<ISequencer>& InSequencer, FGuid BindingID) = 0;
+};
+
+#undef UE_API
+
diff --git a/Engine/Source/Editor/MovieSceneTools/Public/TrackEditors/CommonAnimationTrackEditor.h b/Engine/Source/Editor/MovieSceneTools/Public/TrackEditors/CommonAnimationTrackEditor.h
--- a/Engine/Source/Editor/MovieSceneTools/Public/TrackEditors/CommonAnimationTrackEditor.h
+++ b/Engine/Source/Editor/MovieSceneTools/Public/TrackEditors/CommonAnimationTrackEditor.h
@@ -7,12 +7,15 @@
 #include "Templates/SubclassOf.h"
 #include "Templates/UniquePtr.h"
 #include "Widgets/SWidget.h"
+#include "Widgets/Input/SCheckBox.h"
 #include "ISequencer.h"
 #include "MovieSceneTrack.h"
 #include "ISequencerSection.h"
 #include "MovieSceneTrackEditor.h"
 #include "SequencerCoreFwd.h"
 #include "IMovieSceneAnimMixerItemMenuProvider.h"
+#include "IMovieSceneLinkedAnimTrackProvider.h"
+#include "LevelSequenceAnimSequenceLink.h"
 
 #define UE_API MOVIESCENETOOLS_API
 
@@ -22,10 +25,12 @@ struct FMovieSceneSequenceTransform;
 
 class FMenuBuilder;
 class FSequencerSectionPainter;
+class SHorizontalBox;
 class UAnimSeqExportOption;
 class UAnimSequenceBase;
 class UMovieSceneCommonAnimationTrack;
 class UMovieSceneSkeletalAnimationSection;
+class UMovieSceneSkeletalAnimationTrack;
 class USkeletalMeshComponent;
 class USkeleton;
 
@@ -66,7 +71,8 @@ public:
 
 	static UE_API USkeletalMeshComponent* AcquireSkeletalMeshFromObjectGuid(const FGuid& Guid, TSharedPtr<ISequencer> SequencerPtr);
 	static UE_API USkeleton* AcquireSkeletonFromObjectGuid(const FGuid& Guid, TSharedPtr<ISequencer> SequencerPtr);
-
+	static UE_API bool AcquireAllBoundSkeletalMeshObjectsFromObjectGuid(const FGuid& Guid, TSharedPtr<ISequencer> SequencerPtr, TArray<USkeletalMeshComponent*>& OutSkelMeshComponents,
+		TArray<FGuid>& OutBindingIds);
 public:
 
 	// ISequencerTrackEditor interface
@@ -74,6 +80,11 @@ public:
 	UE_API virtual FText GetDisplayName() const override;
 	UE_API virtual void BuildObjectBindingContextMenu(FMenuBuilder& MenuBuilder, const TArray<FGuid>& ObjectBindings, const UClass* ObjectClass) override;
 	UE_API virtual void BuildObjectBindingTrackMenu(FMenuBuilder& MenuBuilder, const TArray<FGuid>& ObjectBindings, const UClass* ObjectClass) override;
+	UE_API virtual void BuildObjectBindingColumnWidgets(
+		TFunctionRef<TSharedRef<SHorizontalBox>()> GetEditBox,
+		const UE::Sequencer::TViewModelPtr<UE::Sequencer::FObjectBindingModel>& ObjectBinding,
+		const UE::Sequencer::FCreateOutlinerViewParams& InParams,
+		const FName& InColumnName) override;
 	UE_API virtual bool HandleAssetAdded(UObject* Asset, const FGuid& TargetObjectGuid) override;
 	UE_API virtual TSharedRef<ISequencerSection> MakeSectionInterface( UMovieSceneSection& SectionObject, UMovieSceneTrack& Track, FGuid ObjectBinding ) override;
 	UE_API virtual TSharedPtr<SWidget> BuildOutlinerEditWidget(const FGuid& ObjectBinding, UMovieSceneTrack* Track, const FBuildEditWidgetParams& Params) override;
@@ -125,10 +136,49 @@ protected:
 	/** Can Open the linked Anim Sequence*/
 	UE_API bool CanOpenLinkedAnimSequence(FGuid Binding);
 
+	/** Set Up The AutoBake for the specified bindings*/
+	UE_API void ConfigureAutoBake(TArray<FGuid> Bindings);
+
+	/** If we can open a linked anim sequence on any binding*/
+	UE_API bool CanOpenLinkedAnimSequences(TArray<FGuid> Bindings);
+
+	/** If bIsolate is true it will activate any linked and autobaked anim tracks, and deactivate control rig or non linked anim tracks, if false does opposite */
+	UE_API void IsolateLinkedAnimTracks(TArray<FGuid> Bindings, bool bIsolate); 
+
+	/** If we can isolate baked anim tracks since they exists*/
+	UE_API bool HasNoLinkedTracks(TArray<FGuid> Bindings);
+
 	friend class FMovieSceneSkeletalAnimationParamsDetailCustomization;
 
+public:
+
+	/** Create a linked animation sequence for the binding. Used by linked anim track providers. */
+	UE_API void CreateLinkedAnimationSequence(USkeletalMeshComponent* SkelMeshComp, USkeleton* Skeleton, FGuid Binding, bool bCreateSoftLink);
+	UE_API static TArray<FLevelSequenceAnimSequenceLinkItem*> GetLinkedAnimSequences(const TWeakPtr<ISequencer>& InSequencer, FGuid Binding);
+	//isolate the linked anim track by muting other tracks here, return true if items were changed
+	UE_API static bool IsolateLinkedAnimTrack(const TWeakPtr<ISequencer>& InSequencer, FGuid BindingID, bool bIsolate);
+
+	// update linked anim sections ranges after bake in case playback range changed
+	UE_API static void UpdateLinkedAnimSectionRange(const TWeakPtr<ISequencer>& InSequencer, FGuid BindingID);
+
+	/** Find the first linked anim track provider that can handle this binding. */
+	UE_API static IMovieSceneLinkedAnimTrackProvider* FindProviderForBinding(const TWeakPtr<ISequencer>& InSequencer, FGuid BindingID);
+
+	/** Find an empty skeletal animation track on the binding. Used by linked anim track providers. */
+	UE_API UMovieSceneSkeletalAnimationTrack* FindEmptyAnimationTrack(FGuid BindingID);
+
+
+private:
+	//save link
+	void SaveAnimSequenceLink(FLevelSequenceAnimSequenceLinkItem& Item);
+	UMovieSceneTrack* GetLinkedAnimTrack(FGuid Binding, bool bCreateIfMissing);
+
+	ECheckBoxState IsLinkedAnimTrackIsolated(FGuid ObjectGuid) const;
+	void OnIsolateAnimTrackClicked(ECheckBoxState NewState, FGuid ObjectGuid);
+
+
 protected:
-	/* Was part of the the section but should be at the track level since it takes the final blended result at the current time, not the section instance value*/
+	/* Was part of the thesection but should be at the track level since it takes the final blended result at the current time, not the section instance value*/
 	UE_API bool CreatePoseAsset(const TArray<UObject*> NewAssets, FGuid InObjectBinding);
 	UE_API void HandleCreatePoseAsset(FGuid InObjectBinding);
 	UE_API bool CanCreatePoseAsset(FGuid InObjectBinding) const;
diff --git a/Engine/Source/Runtime/LevelSequence/Public/LevelSequenceAnimSequenceLink.h b/Engine/Source/Runtime/LevelSequence/Public/LevelSequenceAnimSequenceLink.h
--- a/Engine/Source/Runtime/LevelSequence/Public/LevelSequenceAnimSequenceLink.h
+++ b/Engine/Source/Runtime/LevelSequence/Public/LevelSequenceAnimSequenceLink.h
@@ -28,6 +28,10 @@ public:
 	UPROPERTY(BlueprintReadWrite, Category = Property)
 	FSoftObjectPath PathToAnimSequence;
 
+	//if true it will be autobaked
+	UPROPERTY(BlueprintReadWrite, Category = Property)
+	bool bAutoBake = false;
+
 	//From Editor Only UAnimSeqExportOption we cache this since we can re-import dynamically
 	UPROPERTY(BlueprintReadWrite, Category = Property)
 	bool bExportTransforms = true;
diff --git a/Engine/Source/Runtime/MovieSceneTracks/Private/Tracks/MovieSceneSkeletalAnimationTrack.cpp b/Engine/Source/Runtime/MovieSceneTracks/Private/Tracks/MovieSceneSkeletalAnimationTrack.cpp
--- a/Engine/Source/Runtime/MovieSceneTracks/Private/Tracks/MovieSceneSkeletalAnimationTrack.cpp
+++ b/Engine/Source/Runtime/MovieSceneTracks/Private/Tracks/MovieSceneSkeletalAnimationTrack.cpp
@@ -28,6 +28,8 @@
 
 #define LOCTEXT_NAMESPACE "MovieSceneSkeletalAnimationTrack"
 
+FColor UMovieSceneSkeletalAnimationTrack::LinkedAnimSeqTrackColor = FColor(60, 4, 46);
+
 /* UMovieSceneSkeletalAnimationTrack structors
  *****************************************************************************/
 
diff --git a/Engine/Source/Runtime/MovieSceneTracks/Public/Tracks/MovieSceneSkeletalAnimationTrack.h b/Engine/Source/Runtime/MovieSceneTracks/Public/Tracks/MovieSceneSkeletalAnimationTrack.h
--- a/Engine/Source/Runtime/MovieSceneTracks/Public/Tracks/MovieSceneSkeletalAnimationTrack.h
+++ b/Engine/Source/Runtime/MovieSceneTracks/Public/Tracks/MovieSceneSkeletalAnimationTrack.h
@@ -51,6 +51,9 @@ public:
 	UFUNCTION()
 	MOVIESCENETRACKS_API ESwapRootBone GetSwapRootBone() const;
 
+public:
+	static MOVIESCENETRACKS_API FColor LinkedAnimSeqTrackColor;
+
 };
 
 
