Sequencer: Significantly improve performance of DeleteKeys and ReplaceKeys.

DeleteKeys: Performance gains come from removing keys in reverse order. This avoids expensive RemoveAt calls shuffling all the key indices backwards (n*n iterations vs just n).

ReplaceKeys: Performance gains come from factoring PostProcessChanges out of the for loop and invoking it only once. This was likely a copy-paste mistake from my side when I implemented it.

=== Performance data ===
Test case: ControlRigExample -> Performance_LS -> global_ctrl.Location.X -> Delete all keys -> Undo.
- RemoveKeys: before 5s and after instant
- ReplaceKeys: before 14.6s and after instant

#Synced-CL 52378315

diff --git a/Engine/Source/Editor/CurveEditor/Private/CurveModel.cpp b/Engine/Source/Editor/CurveEditor/Private/CurveModel.cpp
--- a/Engine/Source/Editor/CurveEditor/Private/CurveModel.cpp
+++ b/Engine/Source/Editor/CurveEditor/Private/CurveModel.cpp
@@ -20,6 +20,8 @@ void FCurveModel::PutKeys(
 	TConstArrayView<FKeyHandle> InKeys, TConstArrayView<FKeyPosition> InKeyPositions, TConstArrayView<FKeyAttributes> InAttributes
 	)
 {
+	TRACE_CPUPROFILER_EVENT_SCOPE(FCurveModel::PutKeys)
+	
 	const int32 Num = InKeys.Num();
 	if (ensure(InKeyPositions.Num() == Num && InAttributes.Num() == Num))
 	{
diff --git a/Engine/Source/Editor/MovieSceneTools/Private/Channels/BezierChannelCurveModel.cpp b/Engine/Source/Editor/MovieSceneTools/Private/Channels/BezierChannelCurveModel.cpp
--- a/Engine/Source/Editor/MovieSceneTools/Private/Channels/BezierChannelCurveModel.cpp
+++ b/Engine/Source/Editor/MovieSceneTools/Private/Channels/BezierChannelCurveModel.cpp
@@ -363,6 +363,7 @@ void FBezierChannelCurveModel<ChannelType, ChannelValue, KeyType>::GetKeyAttribu
 template<typename ChannelType, typename ChannelValue, typename KeyType> 
 void FBezierChannelCurveModel<ChannelType, ChannelValue, KeyType>::SetKeyAttributes(TArrayView<const FKeyHandle> InKeys, TArrayView<const FKeyAttributes> InAttributes, EPropertyChangeType::Type ChangeType)
 {
+	TRACE_CPUPROFILER_EVENT_SCOPE(FChannelCurveModel::SetKeyAttributes)
 	UE::MovieScene::FScopedSignedObjectModifyDefer Defer(false /*do not force upate*/);
 
 	ChannelType*             Channel     = this->GetChannelHandle().Get();
diff --git a/Engine/Source/Editor/MovieSceneTools/Private/Channels/ChannelCurveModel.cpp b/Engine/Source/Editor/MovieSceneTools/Private/Channels/ChannelCurveModel.cpp
--- a/Engine/Source/Editor/MovieSceneTools/Private/Channels/ChannelCurveModel.cpp
+++ b/Engine/Source/Editor/MovieSceneTools/Private/Channels/ChannelCurveModel.cpp
@@ -519,6 +519,7 @@ bool FChannelCurveModel<ChannelType, ChannelValue, KeyType>::Evaluate(double Tim
 template <class ChannelType, class ChannelValue, class KeyType>
 void FChannelCurveModel<ChannelType, ChannelValue, KeyType>::AddKeys(TArrayView<const FKeyPosition> InKeyPositions, TArrayView<const FKeyAttributes> InKeyAttributes, TArrayView<TOptional<FKeyHandle>>* OutKeyHandles)
 {
+	TRACE_CPUPROFILER_EVENT_SCOPE(FChannelCurveModel::AddKeys)
 	check(InKeyPositions.Num() == InKeyAttributes.Num() && (!OutKeyHandles || OutKeyHandles->Num() == InKeyPositions.Num()));
 
 	ChannelType* Channel = ChannelHandle.Get();
@@ -544,8 +545,12 @@ void FChannelCurveModel<ChannelType, ChannelValue, KeyType>::AddKeys(TArrayView<
 			MaxFrame = FMath::Max(MaxFrame, Time);
 
 			ChannelValue Value = (ChannelValue)(Position.OutputValue);
-
-			FKeyHandle NewHandle = ChannelData.UpdateOrAddKey(Time, Value);
+			const FKeyHandle NewHandle = [&]
+			{
+				TRACE_CPUPROFILER_EVENT_SCOPE(FChannelCurveModel::SetKeyAttributes);
+				return ChannelData.UpdateOrAddKey(Time, Value);
+			}();
+			
 			if (NewHandle != FKeyHandle::Invalid())
 			{
 				NewKeyHandles[Index] = NewHandle;
@@ -574,6 +579,8 @@ void FChannelCurveModel<ChannelType, ChannelValue, KeyType>::AddKeys(TArrayView<
 template <class ChannelType, class ChannelValue, class KeyType>
 void FChannelCurveModel<ChannelType, ChannelValue, KeyType>::RemoveKeys(TArrayView<const FKeyHandle> InKeys, double InCurrentTime)
 {
+	TRACE_CPUPROFILER_EVENT_SCOPE(FChannelCurveModel::RemoveKeys)
+	
 	ChannelType* Channel = ChannelHandle.Get();
 	UMovieSceneSection* Section = WeakSection.Get();
 	if (Channel && Section && !IsReadOnly())
@@ -582,11 +589,16 @@ void FChannelCurveModel<ChannelType, ChannelValue, KeyType>::RemoveKeys(TArrayVi
 		double CurrentValue = 0.0;
 		Evaluate(InCurrentTime, CurrentValue);
 
-		for (FKeyHandle Handle : InKeys)
+		// Iterate in reverse so RemoveAt always operates at or after the current index, avoiding O(n²) index shifting.
+		// TMovieSceneChannelData stores keys sorted by time, so back-to-front removal keeps earlier indices stable.
+		// InKeys is usually passed in pre-sorted by ascending time, which we abuse here for performance gains.
+		for (int32 Index = InKeys.Num() - 1; Index >= 0; --Index)
 		{
-			int32 KeyIndex = ChannelData.GetIndex(Handle);
+			const FKeyHandle Handle = InKeys[Index];
+			const int32 KeyIndex = ChannelData.GetIndex(Handle);
 			if (KeyIndex != INDEX_NONE)
 			{
+				TRACE_CPUPROFILER_EVENT_SCOPE(FChannelCurveModel::RemoveKey)
 				ChannelData.RemoveKey(KeyIndex);
 			}
 		}
@@ -606,6 +618,7 @@ void FChannelCurveModel<ChannelType, ChannelValue, KeyType>::ReplaceKeyHandles(
 	TConstArrayView<FKeyHandle> InCurrentHandles, TConstArrayView<FKeyHandle> InNewHandles
 	)
 {
+	TRACE_CPUPROFILER_EVENT_SCOPE(FChannelCurveModel::ReplaceKeyHandles)
 	if (!ensure(InCurrentHandles.Num() == InNewHandles.Num()))
 	{
 		return;
@@ -630,8 +643,9 @@ void FChannelCurveModel<ChannelType, ChannelValue, KeyType>::ReplaceKeyHandles(
 			}
 
 			ChannelData.ReplaceKeyHandle(KeyIndex, InNewHandles[Index]);
-			PostProcessChanges(Channel, Section);
 		}
+		
+		PostProcessChanges(Channel, Section);
 	}
 }
 
@@ -695,6 +709,7 @@ void FChannelCurveModel<ChannelType, ChannelValue, KeyType>::GetCurveColorObject
 template <class ChannelType, class ChannelValue, class KeyType>
 void FChannelCurveModel<ChannelType, ChannelValue, KeyType>::PostProcessChanges(TNotNull<ChannelType*> Channel, TNotNull<UMovieSceneSection*> Section)
 {
+	TRACE_CPUPROFILER_EVENT_SCOPE(FChannelCurveModel::PostProcessChanges)
 	LastSignature.Invalidate();
 	
 	// Allow channel to make any follow up changes to the keys before executing any other delegates
