Sequencer Baker: Fix for possible divide by zero
#trivial

#Synced-CL 52664489

diff --git a/Engine/Source/Editor/MovieSceneTools/Private/Baker/SequencerBaker.cpp b/Engine/Source/Editor/MovieSceneTools/Private/Baker/SequencerBaker.cpp
--- a/Engine/Source/Editor/MovieSceneTools/Private/Baker/SequencerBaker.cpp
+++ b/Engine/Source/Editor/MovieSceneTools/Private/Baker/SequencerBaker.cpp
@@ -25,6 +25,10 @@ bool UE::AIE::FBakeTimeRange::UpdateIfNeeded(const TRange<FFrameNumber>& InRange
 		DisplayRate = InDisplayRate;
 		FrameStep = FFrameRate::TransformTime(FFrameNumber(1), InDisplayRate,
 			InTickResolution).RoundToFrame();
+		if (FrameStep.Value <= 0)
+		{
+			FrameStep.Value = 1;
+		}
 		StartFrame = InRange.GetLowerBoundValue();
 		EndFrame = InRange.GetUpperBoundValue();
 		//make sure end frame falls on frame step 
